- Fixing non unity CIS issues

diff --git a/Engine/Plugins/VirtualProduction/StageMonitoring/Source/StageMonitorEditor/Private/Widgets/SDataProviderActivityFilter.cpp b/Engine/Plugins/VirtualProduction/StageMonitoring/Source/StageMonitorEditor/Private/Widgets/SDataProviderActivityFilter.cpp
--- a/Engine/Plugins/VirtualProduction/StageMonitoring/Source/StageMonitorEditor/Private/Widgets/SDataProviderActivityFilter.cpp
+++ b/Engine/Plugins/VirtualProduction/StageMonitoring/Source/StageMonitorEditor/Private/Widgets/SDataProviderActivityFilter.cpp
@@ -2,19 +2,25 @@
 
 #include "SDataProviderActivityFilter.h"
 
+#include "CoreGlobals.h"
 #include "EditorStyleSet.h"
 #include "Framework/MultiBox/MultiBoxBuilder.h"
+#include "Misc/App.h"
+#include "Misc/ConfigCacheIni.h"
 #include "StageMessages.h"
 #include "Styling/SlateTypes.h"
 #include "UObject/StructOnScope.h"
 #include "UObject/UObjectIterator.h"
 #include "Widgets/Input/SComboButton.h"
-#include "Widgets/SBoxPanel.h"
+#include "Widgets/Input/SCheckBox.h"
 #include "Widgets/Input/SNumericEntryBox.h"
+#include "Widgets/SBoxPanel.h"
 
 #define LOCTEXT_NAMESPACE "SDataProviderActivityFilter"
 
 
+static TAutoConsoleVariable<int32> CVarStageMonitorDefaultMaxMessageAge(TEXT("StageMonitor.DefaultMaxMessageAge"), 30, TEXT("The default value in minutes for the maximum age for which to display messages in the data monitor."));
+
 
 FDataProviderActivityFilter::FDataProviderActivityFilter(TWeakPtr<IStageMonitorSession> InSession)
 	: Session(MoveTemp(InSession))
diff --git a/Engine/Plugins/VirtualProduction/StageMonitoring/Source/StageMonitorEditor/Private/Widgets/SDataProviderActivityFilter.h b/Engine/Plugins/VirtualProduction/StageMonitoring/Source/StageMonitorEditor/Private/Widgets/SDataProviderActivityFilter.h
--- a/Engine/Plugins/VirtualProduction/StageMonitoring/Source/StageMonitorEditor/Private/Widgets/SDataProviderActivityFilter.h
+++ b/Engine/Plugins/VirtualProduction/StageMonitoring/Source/StageMonitorEditor/Private/Widgets/SDataProviderActivityFilter.h
@@ -19,7 +19,6 @@ class FStructOnScope;
 class FMenuBuilder;
 class IStageMonitorSession;
 
-static TAutoConsoleVariable<int32> CVarStageMonitorDefaultMaxMessageAge(TEXT("StageMonitor.DefaultMaxMessageAge"), 30, TEXT("The default value in minutes for the maximum age for which to display messages in the data monitor."));
 
 /** Filter settings used live and also load/saved to ini config */
 USTRUCT()
