- Fix crash when opening stage monitor
- Fixed missing icon with placeholder
- Fixed TimedDataMonitor menu entry location to be in VP

diff --git a/Engine/Plugins/VirtualProduction/StageMonitoring/Source/StageMonitorEditor/Private/SStageMonitorPanel.cpp b/Engine/Plugins/VirtualProduction/StageMonitoring/Source/StageMonitorEditor/Private/SStageMonitorPanel.cpp
--- a/Engine/Plugins/VirtualProduction/StageMonitoring/Source/StageMonitorEditor/Private/SStageMonitorPanel.cpp
+++ b/Engine/Plugins/VirtualProduction/StageMonitoring/Source/StageMonitorEditor/Private/SStageMonitorPanel.cpp
@@ -62,7 +62,7 @@ void SStageMonitorPanel::RegisterNomadTabSpawner(TSharedRef<FWorkspaceItem> InWo
 			.SetDisplayName(LOCTEXT("TabTitle", "Stage Monitor"))
 			.SetTooltipText(LOCTEXT("TooltipText", "Monitor performance data from stage machines"))
 			.SetGroup(InWorkspaceItem)
-			.SetIcon(FSlateIcon(FEditorStyle::GetStyleSetName(), "LevelEditor.GameSettings.Small")); //todo t2
+			.SetIcon(FSlateIcon(FEditorStyle::GetStyleSetName(), "UnrealInsights.MenuIcon"));
 	};
 
 	FLevelEditorModule& LevelEditorModule = FModuleManager::LoadModuleChecked<FLevelEditorModule>(StageMonitorUtilities::NAME_LevelEditorModuleName);
diff --git a/Engine/Plugins/VirtualProduction/StageMonitoring/Source/StageMonitorEditor/Private/Widgets/SDataProviderActivities.cpp b/Engine/Plugins/VirtualProduction/StageMonitoring/Source/StageMonitorEditor/Private/Widgets/SDataProviderActivities.cpp
--- a/Engine/Plugins/VirtualProduction/StageMonitoring/Source/StageMonitorEditor/Private/Widgets/SDataProviderActivities.cpp
+++ b/Engine/Plugins/VirtualProduction/StageMonitoring/Source/StageMonitorEditor/Private/Widgets/SDataProviderActivities.cpp
@@ -53,6 +53,29 @@ void SDataProviderActivities::Construct(const FArguments& InArgs, TSharedPtr<SSt
 	StructureDetailsView->GetDetailsView()->SetIsPropertyEditingEnabledDelegate(FIsPropertyEditingEnabled::CreateLambda([]() { return false; }));
 
 
+	ActivityList = SNew(SListView<FDataProviderActivityPtr>)
+		.ListItemsSource(&FilteredActivities)
+		.OnGenerateRow(this, &SDataProviderActivities::OnGenerateActivityRowWidget)
+		.SelectionMode(ESelectionMode::Single)
+		.OnSelectionChanged(this, &SDataProviderActivities::OnListViewSelectionChanged)
+		//.AllowOverscroll(EAllowOverscroll::No)
+		.HeaderRow
+		(
+			SNew(SHeaderRow)
+			+ SHeaderRow::Column(DataProviderActivitiesListView::HeaderIdName_Timecode)
+			.FillWidth(15.f)
+			.DefaultLabel(LOCTEXT("HeaderName_Timecode", "Timecode"))
+			+ SHeaderRow::Column(DataProviderActivitiesListView::HeaderIdName_StageName)
+			.FillWidth(25.f)
+			.DefaultLabel(LOCTEXT("HeaderName_StageName", "Stage Name"))
+			+ SHeaderRow::Column(DataProviderActivitiesListView::HeaderIdName_Type)
+			.FillWidth(25.f)
+			.DefaultLabel(LOCTEXT("HeaderName_Type", "Type"))
+			+ SHeaderRow::Column(DataProviderActivitiesListView::HeaderIdName_Description)
+			.FillWidth(35.f)
+			.DefaultLabel(LOCTEXT("HeaderName_Description", "Description"))
+		);
+
 	ChildSlot
 	[
 		SNew(SSplitter)
@@ -76,28 +99,7 @@ void SDataProviderActivities::Construct(const FArguments& InArgs, TSharedPtr<SSt
 			.VAlign(VAlign_Fill)
 			.Padding(10.f, 0.f, 10.f, 10.f)
 			[
-				SAssignNew(ActivityList, SListView<FDataProviderActivityPtr>)
-				.ListItemsSource(&FilteredActivities)
-				.OnGenerateRow(this, &SDataProviderActivities::OnGenerateActivityRowWidget)
-				.SelectionMode(ESelectionMode::Single)
-				.OnSelectionChanged(this, &SDataProviderActivities::OnListViewSelectionChanged)
-				//.AllowOverscroll(EAllowOverscroll::No)
-				.HeaderRow
-				(
-					SNew(SHeaderRow)
-					+ SHeaderRow::Column(DataProviderActivitiesListView::HeaderIdName_Timecode)
-					.FillWidth(15.f)
-					.DefaultLabel(LOCTEXT("HeaderName_Timecode", "Timecode"))
-					+ SHeaderRow::Column(DataProviderActivitiesListView::HeaderIdName_StageName)
-					.FillWidth(25.f)
-					.DefaultLabel(LOCTEXT("HeaderName_StageName", "Stage Name"))
-					+ SHeaderRow::Column(DataProviderActivitiesListView::HeaderIdName_Type)
-					.FillWidth(25.f)
-					.DefaultLabel(LOCTEXT("HeaderName_Type", "Type"))
-					+ SHeaderRow::Column(DataProviderActivitiesListView::HeaderIdName_Description)
-					.FillWidth(35.f)
-					.DefaultLabel(LOCTEXT("HeaderName_Description", "Description"))
-				)
+				ActivityList.ToSharedRef()
 			]
 		]
 		+ SSplitter::Slot()
diff --git a/Engine/Plugins/VirtualProduction/TimedDataMonitor/Source/TimedDataMonitorEditor/Private/TimedDataMonitorEditorModule.cpp b/Engine/Plugins/VirtualProduction/TimedDataMonitor/Source/TimedDataMonitorEditor/Private/TimedDataMonitorEditorModule.cpp
--- a/Engine/Plugins/VirtualProduction/TimedDataMonitor/Source/TimedDataMonitorEditor/Private/TimedDataMonitorEditorModule.cpp
+++ b/Engine/Plugins/VirtualProduction/TimedDataMonitor/Source/TimedDataMonitorEditor/Private/TimedDataMonitorEditorModule.cpp
@@ -20,7 +20,7 @@ class FTimedDataMonitorEditorModule : public IModuleInterface
 public:
 	virtual void StartupModule() override
 	{
-		STimedDataMonitorPanel::RegisterNomadTabSpawner(WorkspaceMenu::GetMenuStructure().GetToolsCategory());
+		STimedDataMonitorPanel::RegisterNomadTabSpawner(WorkspaceMenu::GetMenuStructure().GetLevelEditorVirtualProductionCategory());
 
 		if (ISettingsModule* SettingsModule = FModuleManager::GetModulePtr<ISettingsModule>("Settings"))
 		{
