- Changing stagemonitor role exclusion to only filter when it's auto started to be able to start it once in editor independantly.

diff --git a/Engine/Plugins/VirtualProduction/StageMonitoring/Source/StageMonitor/Private/StageMonitor.cpp b/Engine/Plugins/VirtualProduction/StageMonitoring/Source/StageMonitor/Private/StageMonitor.cpp
--- a/Engine/Plugins/VirtualProduction/StageMonitoring/Source/StageMonitor/Private/StageMonitor.cpp
+++ b/Engine/Plugins/VirtualProduction/StageMonitoring/Source/StageMonitor/Private/StageMonitor.cpp
@@ -13,7 +13,6 @@
 #include "StageMonitoringSettings.h"
 #include "StageMonitorModule.h"
 #include "StageMonitorUtils.h"
-#include "VPSettings.h"
 
 
 namespace StageMonitorUtils
@@ -42,8 +41,7 @@ void FStageMonitor::Initialize()
 
 void FStageMonitor::Start()
 {
-	const UStageMonitoringSettings* Settings = GetDefault<UStageMonitoringSettings>();
-	if (!Settings->MonitorSettings.bUseRoleFiltering || GetDefault<UVPSettings>()->GetRoles().HasAny(Settings->MonitorSettings.SupportedRoles))
+	if (bIsActive == false)
 	{
 		const FString MessageEndpointName = TEXT("StageDataMonitor");
 		MonitorEndpoint = FMessageEndpoint::Builder(*MessageEndpointName)
@@ -61,12 +59,6 @@ void FStageMonitor::Start()
 
 		bIsActive = true;
 	}
-	else
-	{
-		UE_LOG(LogStageMonitor, Log, TEXT("Can't start StageMonitor. Role filtering is enabled and our roles (%s) are filtered out (%s)")
-		, *GetDefault<UVPSettings>()->GetRoles().ToStringSimple()
-		, *Settings->MonitorSettings.SupportedRoles.ToStringSimple())
-	}
 }
 
 void FStageMonitor::Stop()
diff --git a/Engine/Plugins/VirtualProduction/StageMonitoring/Source/StageMonitor/Private/StageMonitorModule.cpp b/Engine/Plugins/VirtualProduction/StageMonitoring/Source/StageMonitor/Private/StageMonitorModule.cpp
--- a/Engine/Plugins/VirtualProduction/StageMonitoring/Source/StageMonitor/Private/StageMonitorModule.cpp
+++ b/Engine/Plugins/VirtualProduction/StageMonitoring/Source/StageMonitor/Private/StageMonitorModule.cpp
@@ -7,6 +7,7 @@
 #include "StageMonitor.h"
 #include "StageMonitoringSettings.h"
 #include "StageMonitorSessionManager.h"
+#include "VPSettings.h"
 
 const FName IStageMonitorModule::ModuleName = TEXT("StageMonitor");
 
@@ -40,10 +41,20 @@ void FStageMonitorModule::OnEngineLoopInitComplete()
 	SessionManager = MakeUnique<FStageMonitorSessionManager>();
 	StageMonitor = MakeUnique<FStageMonitor>();
 	StageMonitor->Initialize();
+
 	const UStageMonitoringSettings* Settings = GetDefault<UStageMonitoringSettings>();
 	if (Settings->MonitorSettings.ShouldAutoStartOnLaunch())
 	{
-		EnableMonitor(true);
+		if (!Settings->MonitorSettings.bUseRoleFiltering || GetDefault<UVPSettings>()->GetRoles().HasAny(Settings->MonitorSettings.SupportedRoles))
+		{
+			EnableMonitor(true);
+		}
+		else
+		{
+			UE_LOG(LogStageMonitor, Log, TEXT("Can't autostart StageMonitor. Role filtering is enabled and our roles (%s) are filtered out (%s)")
+				, *GetDefault<UVPSettings>()->GetRoles().ToStringSimple()
+				, *Settings->MonitorSettings.SupportedRoles.ToStringSimple())
+		}
 	}
 }
 
diff --git a/Engine/Plugins/VirtualProduction/StageMonitoring/Source/StageMonitorCommon/Public/StageMonitoringSettings.h b/Engine/Plugins/VirtualProduction/StageMonitoring/Source/StageMonitorCommon/Public/StageMonitoringSettings.h
--- a/Engine/Plugins/VirtualProduction/StageMonitoring/Source/StageMonitorCommon/Public/StageMonitoringSettings.h
+++ b/Engine/Plugins/VirtualProduction/StageMonitoring/Source/StageMonitorCommon/Public/StageMonitoringSettings.h
@@ -147,7 +147,10 @@ public:
 	bool ShouldAutoStartOnLaunch() const;
 
 public:
-	/** If true, Monitor will only start if machine has a role contained in SupportedRoles */
+	/** 
+	 * If true, Monitor will only autostart if machine has a role contained in SupportedRoles 
+	 * Once in editor, you can always start a monitor manually independently of the roles.
+	 */
 	UPROPERTY(config, EditAnywhere, Category = "Settings", meta = (InlineEditConditionToggle))
 	bool bUseRoleFiltering = false;
 
diff --git a/Engine/Plugins/VirtualProduction/StageMonitoring/Source/StageMonitorEditor/Private/Widgets/SDataProviderActivityFilter.cpp b/Engine/Plugins/VirtualProduction/StageMonitoring/Source/StageMonitorEditor/Private/Widgets/SDataProviderActivityFilter.cpp
--- a/Engine/Plugins/VirtualProduction/StageMonitoring/Source/StageMonitorEditor/Private/Widgets/SDataProviderActivityFilter.cpp
+++ b/Engine/Plugins/VirtualProduction/StageMonitoring/Source/StageMonitorEditor/Private/Widgets/SDataProviderActivityFilter.cpp
@@ -88,6 +88,39 @@ FDataProviderActivityFilter::EFilterResult FDataProviderActivityFilter::FilterAc
 		{
 			return EFilterResult::FailRestrictedProviders;
 		}
+
+		if (FilterSettings.RestrictedRoles.Num() > 0)
+		{
+			if (CachedRoleStringToArray.Contains(ProviderEntry.Descriptor.RolesStringified) == false)
+			{
+				TArray<FString> RoleArray;
+				ProviderEntry.Descriptor.RolesStringified.ParseIntoArray(RoleArray, TEXT(","));
+				TArray<FName> ToCache;
+				ToCache.Reserve(RoleArray.Num());
+				for (const FString& RoleString : RoleArray)
+				{
+					ToCache.Add(FName(RoleString));
+				}
+
+				CachedRoleStringToArray.Add(ProviderEntry.Descriptor.RolesStringified, ToCache);
+			}
+
+			bool bHasOneValidRole = false;
+			const TArray<FName>& CachedRoles = CachedRoleStringToArray.FindChecked(ProviderEntry.Descriptor.RolesStringified);
+			for (FName Role : CachedRoles)
+			{
+				if (!FilterSettings.RestrictedRoles.Contains(Role))
+				{
+					bHasOneValidRole = true;
+					break;
+				}
+			}
+
+			if (bHasOneValidRole == false)
+			{
+				return EFilterResult::FailRestrictedRoles;
+			}
+		}
 	}
 
 	if (FilterSettings.bEnableTimeFilter)
@@ -296,6 +329,15 @@ TSharedRef<SWidget> SDataProviderActivityFilter::MakeAddFilterMenu()
 			EUserInterfaceActionType::None
 		);
 
+		MenuBuilder.AddSubMenu(
+			LOCTEXT("FilterRoles", "Role"),
+			LOCTEXT("FilterRolesToolTip", "Filters based on the provider role"),
+			FNewMenuDelegate::CreateSP(this, &SDataProviderActivityFilter::CreateRoleFilterMenu),
+			FUIAction(),
+			NAME_None,
+			EUserInterfaceActionType::None
+		);
+
 		MenuBuilder.AddSubMenu(
 			LOCTEXT("FilterCriticalSources", "Critical state source"),
 			LOCTEXT("FilterCriticalSourcesTooltip", "Filters based on critical state sources"),
@@ -383,14 +425,12 @@ void SDataProviderActivityFilter::CreateProviderFilterMenu(FMenuBuilder& MenuBui
 				);
 			};
 
-			const TConstArrayView<FStageSessionProviderEntry> Providers = SessionPtr->GetProviders();
-			for (const FStageSessionProviderEntry& Provider : Providers)
+			for (const FStageSessionProviderEntry& Provider : SessionPtr->GetProviders())
 			{
 				CreateMenuEntry(Provider);
 			}
 			
-			const TConstArrayView<FStageSessionProviderEntry> ClearedProviders = SessionPtr->GetClearedProviders();
-			for (const FStageSessionProviderEntry& Provider : ClearedProviders)
+			for (const FStageSessionProviderEntry& Provider : SessionPtr->GetClearedProviders())
 			{
 				CreateMenuEntry(Provider);
 			}
@@ -399,6 +439,76 @@ void SDataProviderActivityFilter::CreateProviderFilterMenu(FMenuBuilder& MenuBui
 	MenuBuilder.EndSection();
 }
 
+void SDataProviderActivityFilter::ToggleRoleFilter(FName RoleName)
+{
+	if (CurrentFilter->FilterSettings.RestrictedRoles.Contains(RoleName))
+	{
+		CurrentFilter->FilterSettings.RestrictedRoles.RemoveSingle(RoleName);
+	}
+	else
+	{
+		CurrentFilter->FilterSettings.RestrictedRoles.AddUnique(RoleName);
+	}
+
+	OnActivityFilterChanged.ExecuteIfBound();
+}
+
+bool SDataProviderActivityFilter::IsRoleFiltered(FName RoleName) const
+{
+	// The list contains types to filter. Inverse return value to display a more natural way of looking at filter choices
+	return !CurrentFilter->FilterSettings.RestrictedRoles.Contains(RoleName);
+}
+
+void SDataProviderActivityFilter::CreateRoleFilterMenu(FMenuBuilder& MenuBuilder)
+{
+	MenuBuilder.BeginSection("AvailableRoles", LOCTEXT("AvailableRoles", "Roles"));
+	{
+		if (const TSharedPtr<IStageMonitorSession> SessionPtr = Session.Pin())
+		{
+			//Build list of roles found in active and cleared providers
+			TSet<FName> AvailableRoles;
+			for (const FStageSessionProviderEntry& Provider : SessionPtr->GetProviders())
+			{
+				TArray<FString> ProviderRoles;
+				Provider.Descriptor.RolesStringified.ParseIntoArray(ProviderRoles, TEXT(","));
+				for (const FString& Role : ProviderRoles)
+				{
+					AvailableRoles.Add(*Role);
+				}
+			}
+
+			for (const FStageSessionProviderEntry& Provider : SessionPtr->GetClearedProviders())
+			{
+				TArray<FString> ProviderRoles;
+				Provider.Descriptor.RolesStringified.ParseIntoArray(ProviderRoles, TEXT(","));
+				for (const FString& Role : ProviderRoles)
+				{
+					AvailableRoles.Add(*Role);
+				}
+			}
+
+			for (FName Role : AvailableRoles)
+			{
+				MenuBuilder.AddMenuEntry
+				(
+					FText::FromName(Role), //Label
+					FText::GetEmpty(), //Tooltip
+					FSlateIcon(),
+					FUIAction
+					(
+						FExecuteAction::CreateSP(this, &SDataProviderActivityFilter::ToggleRoleFilter, Role)
+						, FCanExecuteAction()
+						, FIsActionChecked::CreateSP(this, &SDataProviderActivityFilter::IsRoleFiltered, Role)
+					),
+					NAME_None,
+					EUserInterfaceActionType::ToggleButton
+				);
+			}
+		}
+	}
+	MenuBuilder.EndSection();
+}
+
 void SDataProviderActivityFilter::CreateCriticalStateSourceFilterMenu(FMenuBuilder& MenuBuilder)
 {
 	MenuBuilder.BeginSection("CriticalStateSources", LOCTEXT("CriticalStateSources", "Critical state sources"));
diff --git a/Engine/Plugins/VirtualProduction/StageMonitoring/Source/StageMonitorEditor/Private/Widgets/SDataProviderActivityFilter.h b/Engine/Plugins/VirtualProduction/StageMonitoring/Source/StageMonitorEditor/Private/Widgets/SDataProviderActivityFilter.h
--- a/Engine/Plugins/VirtualProduction/StageMonitoring/Source/StageMonitorEditor/Private/Widgets/SDataProviderActivityFilter.h
+++ b/Engine/Plugins/VirtualProduction/StageMonitoring/Source/StageMonitorEditor/Private/Widgets/SDataProviderActivityFilter.h
@@ -41,6 +41,10 @@ struct FDataProviderActivityFilterSettings
 	UPROPERTY()
 	TArray<FName> RestrictedProviders;
 
+	/** Roles that are filtered using their friendly name */
+	UPROPERTY()
+	TArray<FName> RestrictedRoles;
+
 	/** Critical state sources that are filtered */
 	UPROPERTY()
 	TArray<FName> RestrictedSources;
@@ -75,6 +79,9 @@ public:
 	FDataProviderActivityFilterSettings FilterSettings;
 	TWeakPtr<IStageMonitorSession> Session;
 
+	/** Cache of stringified roles to array of roles to avoid parsing into array constantly */
+	mutable TMap<FString, TArray<FName>> CachedRoleStringToArray;
+
 private:
 	enum class EFilterResult : uint8
 	{
@@ -83,6 +90,7 @@ private:
 		FailRestrictedProviders,
 		FailRestrictedCriticalState,
 		FailMaxAge,
+		FailRestrictedRoles,
 		Pass
 	};
 	
@@ -152,6 +160,15 @@ private:
 	/** Creates the menu listing the different providers filter */
 	void CreateProviderFilterMenu(FMenuBuilder& MenuBuilder);
 
+	/** Toggles state of role filter */
+	void ToggleRoleFilter(FName RoleName);
+
+	/** Returns true if role is currently filtered out */
+	bool IsRoleFiltered(FName RoleName) const;
+	
+	/** Creates the menu listing the different roles filter */
+	void CreateRoleFilterMenu(FMenuBuilder& MenuBuilder);
+
 	/** Create the menu listing all critical state sources */
 	void CreateCriticalStateSourceFilterMenu(FMenuBuilder& MenuBuilder);
 
