-  Fixing sorting for stage monitor providers

diff --git a/Engine/Plugins/VirtualProduction/StageMonitoring/Source/StageMonitorEditor/Private/Widgets/SDataProviderListView.cpp b/Engine/Plugins/VirtualProduction/StageMonitoring/Source/StageMonitorEditor/Private/Widgets/SDataProviderListView.cpp
--- a/Engine/Plugins/VirtualProduction/StageMonitoring/Source/StageMonitorEditor/Private/Widgets/SDataProviderListView.cpp
+++ b/Engine/Plugins/VirtualProduction/StageMonitoring/Source/StageMonitorEditor/Private/Widgets/SDataProviderListView.cpp
@@ -301,6 +301,9 @@ FText SDataProviderTableRow::GetGPUTiming() const
  */
 void SDataProviderListView::Construct(const FArguments& InArgs, const TWeakPtr<IStageMonitorSession>& InSession)
 {
+	SortedColumnName = DataProviderListView::HeaderIdName_StageName;
+	SortMode = EColumnSortMode::Ascending;
+
 	AttachToMonitorSession(InSession);
 
 	Super::Construct
@@ -315,10 +318,12 @@ void SDataProviderListView::Construct(const FArguments& InArgs, const TWeakPtr<I
 			.CanSelectGeneratedColumn(true) //To show/hide columns
 
 			+ SHeaderRow::Column(DataProviderListView::HeaderIdName_State)
-			.FixedWidth(45.f)
+			.FixedWidth(55.f)
 			.HAlignHeader(HAlign_Center)
 			.HAlignCell(HAlign_Center)
 			.DefaultLabel(LOCTEXT("HeaderName_State", "State"))
+			.SortMode(this, &SDataProviderListView::GetColumnSortMode, DataProviderListView::HeaderIdName_State)
+			.OnSort(this, &SDataProviderListView::OnColumnSortModeChanged)
 
 			+ SHeaderRow::Column(DataProviderListView::HeaderIdName_Timecode)
 			.FillWidth(.2f)
@@ -327,19 +332,27 @@ void SDataProviderListView::Construct(const FArguments& InArgs, const TWeakPtr<I
 			+ SHeaderRow::Column(DataProviderListView::HeaderIdName_MachineName)
 			.FillWidth(.2f)
 			.DefaultLabel(LOCTEXT("HeaderName_MachineName", "Machine"))
+			.SortMode(this, &SDataProviderListView::GetColumnSortMode, DataProviderListView::HeaderIdName_MachineName)
+			.OnSort(this, &SDataProviderListView::OnColumnSortModeChanged)
 
 			+ SHeaderRow::Column(DataProviderListView::HeaderIdName_ProcessId)
 			.FillWidth(.15f)
 			.DefaultLabel(LOCTEXT("HeaderName_ProcessId", "Process Id"))
+			.SortMode(this, &SDataProviderListView::GetColumnSortMode, DataProviderListView::HeaderIdName_ProcessId)
+			.OnSort(this, &SDataProviderListView::OnColumnSortModeChanged)
 
 			+ SHeaderRow::Column(DataProviderListView::HeaderIdName_StageName)
 			.FillWidth(.2f)
 			.DefaultLabel(LOCTEXT("HeaderName_StageName", "Stage Name"))
 			.ShouldGenerateWidget(true) //Can't hide this column
+			.SortMode(this, &SDataProviderListView::GetColumnSortMode, DataProviderListView::HeaderIdName_StageName)
+			.OnSort(this, &SDataProviderListView::OnColumnSortModeChanged)
 
 			+ SHeaderRow::Column(DataProviderListView::HeaderIdName_Roles)
 			.FillWidth(.2f)
 			.DefaultLabel(LOCTEXT("HeaderName_Roles", "Roles"))
+			.SortMode(this, &SDataProviderListView::GetColumnSortMode, DataProviderListView::HeaderIdName_Roles)
+			.OnSort(this, &SDataProviderListView::OnColumnSortModeChanged)
 
 			+ SHeaderRow::Column(DataProviderListView::HeaderIdName_AverageFPS)
 			.FillWidth(.2f)
@@ -437,6 +450,7 @@ void SDataProviderListView::RebuildDataProviderList()
 		}
 	}
 
+	SortProviderList();
 	RequestListRefresh();
 }
 
@@ -458,5 +472,97 @@ void SDataProviderListView::AttachToMonitorSession(const TWeakPtr<IStageMonitorS
 	}
 }
 
+EColumnSortMode::Type SDataProviderListView::GetColumnSortMode(const FName ColumnId) const
+{
+	if (ColumnId == SortedColumnName)
+	{
+		return SortMode;
+	}
+
+	return EColumnSortMode::None;
+}
+
+void SDataProviderListView::OnColumnSortModeChanged(const EColumnSortPriority::Type SortPriority, const FName& ColumnId, const EColumnSortMode::Type InSortMode)
+{
+	SortedColumnName = ColumnId;
+	SortMode = InSortMode;
+
+	SortProviderList();
+	RequestListRefresh();
+}
+
+void SDataProviderListView::SortProviderList()
+{
+	auto Compare = [this](const FDataProviderTableRowDataPtr& Lhs, const FDataProviderTableRowDataPtr& Rhs, const FName& ColumnName, EColumnSortMode::Type CurrentSortMode)
+	{
+		if (ColumnName == DataProviderListView::HeaderIdName_State)
+		{
+			return CurrentSortMode == EColumnSortMode::Ascending ? Lhs->CachedState < Rhs->CachedState : Lhs->CachedState > Rhs->CachedState;
+		}
+		else if (ColumnName == DataProviderListView::HeaderIdName_MachineName)
+		{
+			const int32 CompareResult = Lhs->Descriptor.MachineName.Compare(Rhs->Descriptor.MachineName);
+			return CurrentSortMode == EColumnSortMode::Ascending ? CompareResult < 0 : CompareResult > 0;
+		}
+		else if (ColumnName == DataProviderListView::HeaderIdName_StageName)
+		{
+			const int32 CompareResult = Lhs->Descriptor.FriendlyName.Compare(Rhs->Descriptor.FriendlyName);
+			return CurrentSortMode == EColumnSortMode::Ascending ? CompareResult < 0 : CompareResult > 0;
+		}
+		else if (ColumnName == DataProviderListView::HeaderIdName_ProcessId)
+		{
+			return CurrentSortMode == EColumnSortMode::Ascending ? Lhs->Descriptor.ProcessId < Rhs->Descriptor.ProcessId : Lhs->Descriptor.ProcessId > Rhs->Descriptor.ProcessId;
+		}
+		else if (ColumnName == DataProviderListView::HeaderIdName_Roles)
+		{
+			if (CachedRoleStringToArray.Contains(Lhs->Descriptor.RolesStringified) == false)
+			{
+				TArray<FString> RoleArray;
+				Lhs->Descriptor.RolesStringified.ParseIntoArray(RoleArray, TEXT(","));
+				CachedRoleStringToArray.Add(Lhs->Descriptor.RolesStringified, RoleArray);
+			}
+			if (CachedRoleStringToArray.Contains(Rhs->Descriptor.RolesStringified) == false)
+			{
+				TArray<FString> RoleArray;
+				Rhs->Descriptor.RolesStringified.ParseIntoArray(RoleArray, TEXT(","));
+				CachedRoleStringToArray.Add(Rhs->Descriptor.RolesStringified, RoleArray);
+			}
+
+			const TArray<FString>& LhsRoles = CachedRoleStringToArray.FindChecked(Lhs->Descriptor.RolesStringified);
+			const TArray<FString>& RhsRoles = CachedRoleStringToArray.FindChecked(Rhs->Descriptor.RolesStringified);
+			TArray<FString>::TConstIterator LhsIterator = LhsRoles.CreateConstIterator();
+			TArray<FString>::TConstIterator RhsIterator = RhsRoles.CreateConstIterator();
+			bool bIsValid = false;
+			do
+			{
+				const FString LhsRole = LhsIterator ? *LhsIterator : FString();
+				const FString RhsRole = RhsIterator ? *RhsIterator : FString();
+				const int32 CompareResult = LhsRole.Compare(RhsRole);
+				if (CompareResult != 0)
+				{
+					return CurrentSortMode == EColumnSortMode::Ascending ? CompareResult < 0 : CompareResult > 0;
+				}
+
+				bIsValid = LhsIterator && RhsIterator;
+				++LhsIterator;
+				++RhsIterator;
+
+			} while (bIsValid);
+
+			return CurrentSortMode == EColumnSortMode::Ascending ? true : false;
+		}
+		else
+		{
+			return CurrentSortMode == EColumnSortMode::Ascending ? Lhs->Descriptor.SessionId < Rhs->Descriptor.SessionId : Lhs->Descriptor.SessionId > Rhs->Descriptor.SessionId;
+		}
+	};
+
+	ListItemsSource.StableSort([&](const FDataProviderTableRowDataPtr& Lhs, const FDataProviderTableRowDataPtr& Rhs)
+		{
+			return Compare(Lhs, Rhs, SortedColumnName, SortMode);
+		});
+}
+
 #undef LOCTEXT_NAMESPACE
 
+
diff --git a/Engine/Plugins/VirtualProduction/StageMonitoring/Source/StageMonitorEditor/Private/Widgets/SDataProviderListView.h b/Engine/Plugins/VirtualProduction/StageMonitoring/Source/StageMonitorEditor/Private/Widgets/SDataProviderListView.h
--- a/Engine/Plugins/VirtualProduction/StageMonitoring/Source/StageMonitorEditor/Private/Widgets/SDataProviderListView.h
+++ b/Engine/Plugins/VirtualProduction/StageMonitoring/Source/StageMonitorEditor/Private/Widgets/SDataProviderListView.h
@@ -94,6 +94,15 @@ private:
 	/** Binds to the current session desired delegates */
 	void AttachToMonitorSession(const TWeakPtr<IStageMonitorSession>& NewSession);
 
+	/** Returns sort method for a given column name */
+	EColumnSortMode::Type GetColumnSortMode(const FName ColumnId) const;
+
+	/** Called when sort mode changed for a given column */
+	void OnColumnSortModeChanged(const EColumnSortPriority::Type SortPriority, const FName& ColumnId, const EColumnSortMode::Type InSortMode);
+
+	/** Sorts provider list based on active column sort mode */
+	void SortProviderList();
+
 private:
 	
 	/** Widget and data containing info about what's shown in the listview */
@@ -108,5 +117,14 @@ private:
 
 	/** Timestamp when we last refreshed UI */
 	double LastRefreshTime = 0.0;
+
+	/** Currently sorted column name */
+	FName SortedColumnName;
+	
+	/** Sort mode for the currently sorted column */
+	EColumnSortMode::Type SortMode = EColumnSortMode::Ascending;
+
+	/** Cache of stringified roles to array of roles to avoid parsing into array constantly */
+	TMap<FString, TArray<FString>> CachedRoleStringToArray;
 };
 
