Fix capitalization for Stage Monitor plugin name.

diff --git a/Engine/Plugins/VirtualProduction/StageMonitoring/StageMonitoring.uplugin b/Engine/Plugins/VirtualProduction/StageMonitoring/StageMonitoring.uplugin
--- a/Engine/Plugins/VirtualProduction/StageMonitoring/StageMonitoring.uplugin
+++ b/Engine/Plugins/VirtualProduction/StageMonitoring/StageMonitoring.uplugin
@@ -2,7 +2,7 @@
 	"FileVersion": 3,
 	"Version": 1,
 	"VersionName": "1.0",
-	"FriendlyName": "Stage monitor",
+	"FriendlyName": "Stage Monitor",
 	"Description": "Plugin enabling monitoring in the context of a virtual production stage where multiple machines are in operation",
 	"Category": "Virtual Production",
 	"CreatedBy": "",
@@ -45,4 +45,4 @@
 			"LoadingPhase": "Default"
 		}
 	]
-}
\ No newline at end of file
+}
