Add additional status column on stage monitor for shaders and asset loading.

diff --git a/Engine/Plugins/VirtualProduction/StageMonitoring/Source/StageDataProvider/Private/FramePerformanceProvider.cpp b/Engine/Plugins/VirtualProduction/StageMonitoring/Source/StageDataProvider/Private/FramePerformanceProvider.cpp
--- a/Engine/Plugins/VirtualProduction/StageMonitoring/Source/StageDataProvider/Private/FramePerformanceProvider.cpp
+++ b/Engine/Plugins/VirtualProduction/StageMonitoring/Source/StageDataProvider/Private/FramePerformanceProvider.cpp
@@ -2,6 +2,7 @@
 
 #include "FramePerformanceProvider.h"
 
+#include "CoreGlobals.h"
 #include "EngineGlobals.h"
 #include "Features/IModularFeatures.h"
 #include "IStageDataProvider.h"
@@ -14,29 +15,220 @@
 #include "StageMonitoringSettings.h"
 #include "StageMonitorUtils.h"
 #include "Stats/StatsData.h"
+#include "UObject/PackageReload.h"
 #include "VPSettings.h"
 
+#include "ShaderCompiler.h"
+
+namespace UE::FramePerformanceProvider::Private
+{
+
+int32 ShadersLeftToCompile()
+{
+	if (GShaderCompilingManager && GShaderCompilingManager->IsCompiling())
+	{
+		return GShaderCompilingManager->GetNumRemainingJobs();
+	}
+	return 0;
+}
+
+FFramePerformanceProviderMessage GetLatestPerformanceData()
+{
+	const float GameThreadTime = FPlatformTime::ToMilliseconds(GGameThreadTime);
+	const float RenderThreadTime = FPlatformTime::ToMilliseconds(GRenderThreadTime);
+	const float GPUTime = FPlatformTime::ToMilliseconds(GGPUFrameTime);
+	const float IdleTimeMilli = (FApp::GetIdleTime() * 1000.0);
+	return FFramePerformanceProviderMessage(
+		EStageMonitorNodeStatus::Unknown,
+		GameThreadTime,
+		RenderThreadTime,
+		GPUTime,
+		IdleTimeMilli,
+		ShadersLeftToCompile()
+		);
+}
+
+}
+
+class FFrameProviderRunner : public FRunnable
+{
+public:
+	FFrameProviderRunner()
+	{
+		Thread.Reset(FRunnableThread::Create( this, TEXT( "StageMonitor Frame Provider Thread" ) ));
+	}
+
+	virtual ~FFrameProviderRunner()
+	{
+		if (Thread)
+		{
+			Thread->Kill();
+		}
+	}
+
+	// FRunnable functions
+	virtual uint32 Run() override
+	{
+		bStopped = false;
+		do
+		{
+			IStageDataProvider::SendMessage<FFramePerformanceProviderMessage>(EStageMessageFlags::None, GetFramePerformanceData());
+
+			FPlatformProcess::Sleep( UpdateFrequency );
+		} while (!bStopped);
+		return 0;
+	}
+
+	virtual void Stop() override
+	{
+		bStopped = true;
+	}
+
+	virtual void Exit() override
+	{
+		Stop();
+	}
+
+	// FRunnable
+
+	struct FLoadInfo
+	{
+		EStageMonitorNodeStatus Status;
+		FString AssetName;
+	};
+
+	FFramePerformanceProviderMessage GetFramePerformanceData()
+	{
+		FFramePerformanceProviderMessage OutboundData = UE::FramePerformanceProvider::Private::GetLatestPerformanceData();
+		if (OutboundData.ShadersToCompile > 0)
+		{
+			OutboundData.Status = EStageMonitorNodeStatus::ShaderCompiling;
+		}
+		else
+		{
+			TOptional<FLoadInfo> LatestLoadInfo = GetLoadInfoTop();
+			if (LatestLoadInfo)
+			{
+				OutboundData.Status = LatestLoadInfo->Status;
+				OutboundData.AssetInStatus = LatestLoadInfo->AssetName;
+			}
+			else
+			{
+				OutboundData.Status = EStageMonitorNodeStatus::Ready;
+			}
+		}
+
+		return OutboundData;
+	}
+
+	void PushLoadInfo(FLoadInfo Info)
+	{
+		FScopeLock Lock(&LoadInfoCS);
+		LoadStack.Push(MoveTemp(Info));
+	}
+
+	TOptional<FLoadInfo> GetLoadInfoTop() const
+	{
+		FScopeLock Lock(&LoadInfoCS);
+		if (LoadStack.Num() > 0)
+		{
+			return LoadStack.Top();
+		}
+		return {};
+	}
+
+	TOptional<FLoadInfo> PopLoadInfo()
+	{
+		FScopeLock Lock(&LoadInfoCS);
+		if (LoadStack.Num() > 0)
+		{
+			TOptional<FLoadInfo> Info = LoadStack.Top();
+			LoadStack.Pop();
+			return Info;
+		}
+		return {};
+	}
+
+	void SetUpdateFrequency(float InFrequency)
+	{
+		UpdateFrequency = InFrequency;
+	}
+
+	TUniquePtr<FRunnableThread> Thread;
+	mutable FCriticalSection LoadInfoCS;
+
+	TArray<FLoadInfo> LoadStack;
+
+	float UpdateFrequency = 0.2f; // Default is 200ms;
+	bool bStopped = false;
+};
+
+
 FFramePerformanceProvider::FFramePerformanceProvider()
 {
-	//Verify if conditions are met to enable frame performance messages
-	FCoreDelegates::OnEndFrame.AddRaw(this, &FFramePerformanceProvider::OnEndFrame);
+	FCoreUObjectDelegates::PreLoadMap.AddRaw(this, &FFramePerformanceProvider::HandlePreLoadMap);
+	FCoreUObjectDelegates::PostLoadMapWithWorld.AddRaw(this, &FFramePerformanceProvider::HandlePostLoadMap);
+	FCoreUObjectDelegates::OnPackageReloaded.AddRaw(this, &FFramePerformanceProvider::HandleAssetReload);
 
+	ProviderThread = MakeUnique<FFrameProviderRunner>();
 #if WITH_EDITOR
 	GetMutableDefault<UStageMonitoringSettings>()->OnSettingChanged().AddRaw(this, &FFramePerformanceProvider::OnStageSettingsChanged);
+
+	const double UpdateInterval = GetDefault<UStageMonitoringSettings>()->ProviderSettings.FramePerformanceSettings.UpdateInterval;
+	ProviderThread->SetUpdateFrequency(static_cast<float>(UpdateInterval));
 #endif //WITH_EDITOR
 	EnableHitchDetection(GetDefault<UStageMonitoringSettings>()->ProviderSettings.HitchDetectionSettings.bEnableHitchDetection);
 }
 
 FFramePerformanceProvider::~FFramePerformanceProvider()
 {
-	//Cleanup what could have been registered
-	FCoreDelegates::OnEndFrame.RemoveAll(this);
+	FCoreUObjectDelegates::PreLoadMap.RemoveAll(this);
+	FCoreUObjectDelegates::PostLoadMapWithWorld.RemoveAll(this);
+	FCoreUObjectDelegates::OnPackageReloaded.RemoveAll(this);
+
 	EnableHitchDetection(false);
 }
 
-void FFramePerformanceProvider::OnEndFrame()
+/** Delegate for when package reload occurs. */
+void FFramePerformanceProvider::HandleAssetReload(const EPackageReloadPhase InPackageReloadPhase, FPackageReloadedEvent* InPackageReloadedEvent)
 {
-	UpdateFramePerformance();
+	if (InPackageReloadPhase == EPackageReloadPhase::PrePackageFixup)
+	{
+		FString Name = InPackageReloadedEvent->GetOldPackage()->GetFullName();
+		ProviderThread->PushLoadInfo({EStageMonitorNodeStatus::HotReload, Name});
+		IStageDataProvider::SendMessage<FAssetLoadingStateProviderMessage>(
+			EStageMessageFlags::Reliable, EStageLoadingState::PreLoad, Name);
+	}
+	else if (InPackageReloadPhase == EPackageReloadPhase::PostPackageFixup)
+	{
+		TOptional<FFrameProviderRunner::FLoadInfo> Info = ProviderThread->PopLoadInfo();
+
+		if (ensure(Info))
+		{
+			IStageDataProvider::SendMessage<FAssetLoadingStateProviderMessage>(
+				EStageMessageFlags::Reliable, EStageLoadingState::PostLoad, Info->AssetName);
+		}
+	}
+}
+
+/** Delegate for pre-load map */
+void FFramePerformanceProvider::HandlePreLoadMap(const FString& MapName)
+{
+	ProviderThread->PushLoadInfo({EStageMonitorNodeStatus::LoadingMap, MapName});
+	IStageDataProvider::SendMessage<FAssetLoadingStateProviderMessage>(
+		EStageMessageFlags::Reliable, EStageLoadingState::PreLoad, MapName);
+}
+
+/** Delegate for post load of map. */
+void FFramePerformanceProvider::HandlePostLoadMap(UWorld* /*unused*/)
+{
+	TOptional<FFrameProviderRunner::FLoadInfo> Info = ProviderThread->PopLoadInfo();
+
+	if (ensure(Info))
+	{
+		IStageDataProvider::SendMessage<FAssetLoadingStateProviderMessage>(
+			EStageMessageFlags::Reliable, EStageLoadingState::PostLoad, Info->AssetName);
+	}
 }
 
 void FFramePerformanceProvider::CheckHitches(int64 Frame)
@@ -68,21 +260,6 @@ void FFramePerformanceProvider::CheckHitches(int64 Frame)
 #endif //STATS
 }
 
-void FFramePerformanceProvider::UpdateFramePerformance()
-{
-	const double CurrentTime = FApp::GetCurrentTime();
-	if (CurrentTime - LastFramePerformanceSent >= GetDefault<UStageMonitoringSettings>()->ProviderSettings.FramePerformanceSettings.UpdateInterval)
-	{
-		LastFramePerformanceSent = CurrentTime;
-
-		const float GameThreadTime = FPlatformTime::ToMilliseconds(GGameThreadTime);
-		const float RenderThreadTime = FPlatformTime::ToMilliseconds(GRenderThreadTime);
-		const float GPUTime = FPlatformTime::ToMilliseconds(GGPUFrameTime);
-		const float IdleTimeMilli = (FApp::GetIdleTime() * 1000.0);
-		IStageDataProvider::SendMessage<FFramePerformanceProviderMessage>(EStageMessageFlags::None, GameThreadTime, RenderThreadTime, GPUTime, IdleTimeMilli);
-	}
-}
-
 #if WITH_EDITOR
 void FFramePerformanceProvider::OnStageSettingsChanged(UObject* Object, struct FPropertyChangedEvent& PropertyChangedEvent)
 {
@@ -92,6 +269,11 @@ void FFramePerformanceProvider::OnStageSettingsChanged(UObject* Object, struct F
 	{
 		EnableHitchDetection(GetDefault<UStageMonitoringSettings>()->ProviderSettings.HitchDetectionSettings.bEnableHitchDetection);
 	}
+	if (PropertyName == GET_MEMBER_NAME_CHECKED(FStageFramePerformanceSettings, UpdateInterval))
+	{
+		const double UpdateInterval = GetDefault<UStageMonitoringSettings>()->ProviderSettings.FramePerformanceSettings.UpdateInterval;
+		ProviderThread->SetUpdateFrequency(static_cast<float>(UpdateInterval));
+	}
 }
 #endif //WITH_EDITOR
 
diff --git a/Engine/Plugins/VirtualProduction/StageMonitoring/Source/StageDataProvider/Private/FramePerformanceProvider.h b/Engine/Plugins/VirtualProduction/StageMonitoring/Source/StageDataProvider/Private/FramePerformanceProvider.h
--- a/Engine/Plugins/VirtualProduction/StageMonitoring/Source/StageDataProvider/Private/FramePerformanceProvider.h
+++ b/Engine/Plugins/VirtualProduction/StageMonitoring/Source/StageDataProvider/Private/FramePerformanceProvider.h
@@ -7,8 +7,12 @@
 #include "StageMessages.h"
 #include "StageMonitoringSettings.h"
 
+#include "Templates/UniquePtr.h"
+
 #include "FramePerformanceProvider.generated.h"
 
+class FFrameProviderRunner;
+
 /**
  * Message sent when a hitch was detected on a provider machine.
  */
@@ -59,8 +63,8 @@ public:
 
 	/** FPS correspondent to the timing that triggered the hitch (game or render thread) */
 	UPROPERTY(VisibleAnywhere, Category = "Hitch")
-	float HitchedTimeFPS = 0.f;	
-	
+	float HitchedTimeFPS = 0.f;
+
 	/** Average FPS when hitch occured */
 	UPROPERTY(VisibleAnywhere, Category = "Hitch")
 	float AverageFPS = 0.f;
@@ -77,15 +81,10 @@ public:
 	~FFramePerformanceProvider();
 
 private:
-	/** At end of frame, send a message, if interval met, about this frame data */
-	void OnEndFrame();
 
 	/** Callback by the stat module to detect if a hitch happened */
 	void CheckHitches(int64 Frame);
 
-	/** Sends a message containing information about the frame performane */
-	void UpdateFramePerformance();
-
 	/** Triggered when a stage monitor settings has changed */
 #if WITH_EDITOR
 	void OnStageSettingsChanged(UObject* Object, struct FPropertyChangedEvent& PropertyChangedEvent);
@@ -96,11 +95,19 @@ private:
 
 private:
 
-	/** Timestamp when last frame performance message was sent */
-	double LastFramePerformanceSent = 0.0;
+	/** Delegate for when package reload occurs. */
+	void HandleAssetReload(const EPackageReloadPhase InPackageReloadPhase, FPackageReloadedEvent* InPackageReloadedEvent);
+
+	/** Delegate for pre-load map */
+	void HandlePreLoadMap(const FString& MapName);
+
+	/** Delegate for post load of map. */
+	void HandlePostLoadMap(UWorld* /*unused*/);
 
 	/** Cached settings for hitch detection since it's called on another thread, we can't access UObject */
 	FStageHitchDetectionSettings CachedHitchSettings;
 
 	bool bIsHitchDetectionEnabled = false;
+
+	TUniquePtr<FFrameProviderRunner> ProviderThread;
 };
diff --git a/Engine/Plugins/VirtualProduction/StageMonitoring/Source/StageMonitor/Private/StageMonitor.h b/Engine/Plugins/VirtualProduction/StageMonitoring/Source/StageMonitor/Private/StageMonitor.h
--- a/Engine/Plugins/VirtualProduction/StageMonitoring/Source/StageMonitor/Private/StageMonitor.h
+++ b/Engine/Plugins/VirtualProduction/StageMonitoring/Source/StageMonitor/Private/StageMonitor.h
@@ -11,6 +11,8 @@
 #include "UObject/ObjectMacros.h"
 #include "Containers/Ticker.h"
 
+#include "Templates/Tuple.h"
+
 class FMessageEndpoint;
 struct FStageDataBaseMessage;
 class IMessageContext;
@@ -55,7 +57,7 @@ private:
 
 	/** Handles Close message from a provider shutting down */
 	void HandleProviderCloseMessage(const FStageProviderCloseMessage& Message, const TSharedRef<IMessageContext, ESPMode::ThreadSafe>& Context);
-	
+
 	/**  */
 	bool Tick(float DeltaTime);
 
@@ -74,8 +76,18 @@ private:
 	{
 #if ENABLE_STAGEMONITOR_LOGGING
 		static_assert(TIsDerivedFrom<MessageType, FStageMonitorBaseMessage>::IsDerived, "MessageType must be a FStageMonitorBaseMessage derived UStruct.");
-		MessageType TempObj = MessageType(Forward<Args>(MoveTempIfPossible(args))...);
-		return SendMessageInternal(&TempObj, MessageType::StaticStruct(), Flags);
+		if constexpr(sizeof...(Args) == 1 && std::is_same<MessageType,TTupleElement<0,TTuple<Args...>>>::value)
+		{
+			MessageType TempObj = Forward<MessageType>(MoveTempIfPossible(TTuple<Args...>::Get<0>(Tie(args...))));
+			return SendMessageInternal(&TempObj, MessageType::StaticStruct(), Flags);
+		}
+		else
+		{
+			MessageType TempObj = MessageType(Forward<Args>(MoveTempIfPossible(args))...);
+			return SendMessageInternal(&TempObj, MessageType::StaticStruct(), Flags);
+		}
+
+
 #endif
 		return false;
 	}
diff --git a/Engine/Plugins/VirtualProduction/StageMonitoring/Source/StageMonitorCommon/Public/StageMonitorUtils.h b/Engine/Plugins/VirtualProduction/StageMonitoring/Source/StageMonitorCommon/Public/StageMonitorUtils.h
--- a/Engine/Plugins/VirtualProduction/StageMonitoring/Source/StageMonitorCommon/Public/StageMonitorUtils.h
+++ b/Engine/Plugins/VirtualProduction/StageMonitoring/Source/StageMonitorCommon/Public/StageMonitorUtils.h
@@ -8,6 +8,16 @@
 
 #include "StageMonitorUtils.generated.h"
 
+UENUM()
+enum class EStageMonitorNodeStatus
+{
+	Unknown,
+	LoadingMap,
+	Ready,
+	HotReload,
+	ShaderCompiling
+};
+
 
 /**
  * Message containing information about frame timings.
@@ -22,13 +32,29 @@ public:
 
 	FFramePerformanceProviderMessage() = default;
 
+	UE_DEPRECATED(5.1, "FFramePerformanceProviderMessage constructor is deprecated, please use updated constructor.")
 	FFramePerformanceProviderMessage(float GameThreadTime, float RenderThreadTime, float GPUTime, float IdleTime)
-		: GameThreadMS(GameThreadTime), RenderThreadMS(RenderThreadTime), GPU_MS(GPUTime), IdleTimeMS(IdleTime)
+		: GameThreadMS(GameThreadTime), RenderThreadMS(RenderThreadTime), GPU_MS(GPUTime), IdleTimeMS(IdleTime), ShadersToCompile(0)
+	{
+		extern ENGINE_API float GAverageFPS;
+		AverageFPS = GAverageFPS;
+	}
+
+	FFramePerformanceProviderMessage(EStageMonitorNodeStatus InStatus, float GameThreadTime, float RenderThreadTime, float GPUTime, float IdleTime, int32 InShadersToCompile)
+	: Status(InStatus), GameThreadMS(GameThreadTime), RenderThreadMS(RenderThreadTime), GPU_MS(GPUTime), IdleTimeMS(IdleTime), ShadersToCompile(InShadersToCompile)
 	{
 		extern ENGINE_API float GAverageFPS;
 		AverageFPS = GAverageFPS;
 	}
 
+	/** Average FrameRate read from GAverageFPS */
+	UPROPERTY(VisibleAnywhere, Category = "Performance")
+	EStageMonitorNodeStatus Status = EStageMonitorNodeStatus::Unknown;
+
+	/** Full path name of the asset involved in the status. This will be empty for Ready, Shader Compiles, and Unknown states  */
+	UPROPERTY(VisibleAnywhere, Category = "Performance")
+	FString AssetInStatus;
+
 	/** Average FrameRate read from GAverageFPS */
 	UPROPERTY(VisibleAnywhere, Category = "Performance")
 	float AverageFPS = 0.f;
@@ -48,6 +74,10 @@ public:
 	/** Idle time (slept) in milliseconds during last frame */
 	UPROPERTY(VisibleAnywhere, Category = "Performance", meta = (Unit = "ms"))
 	float IdleTimeMS = 0.f;
+
+	/** Number of shaders currently being compiled. */
+	UPROPERTY(VisibleAnywhere, Category = "Performance")
+	int32 ShadersToCompile = 0;
 };
 
 
diff --git a/Engine/Plugins/VirtualProduction/StageMonitoring/Source/StageMonitorEditor/Private/Widgets/SDataProviderListView.cpp b/Engine/Plugins/VirtualProduction/StageMonitoring/Source/StageMonitorEditor/Private/Widgets/SDataProviderListView.cpp
--- a/Engine/Plugins/VirtualProduction/StageMonitoring/Source/StageMonitorEditor/Private/Widgets/SDataProviderListView.cpp
+++ b/Engine/Plugins/VirtualProduction/StageMonitoring/Source/StageMonitorEditor/Private/Widgets/SDataProviderListView.cpp
@@ -15,6 +15,7 @@
 namespace DataProviderListView
 {
 	const FName HeaderIdName_State = TEXT("State");
+	const FName HeaderIdName_Status = TEXT("Status");
 	const FName HeaderIdName_Timecode = TEXT("Timecode");
 	const FName HeaderIdName_MachineName = TEXT("MachineName");
 	const FName HeaderIdName_ProcessId = TEXT("ProcessId");
@@ -25,6 +26,7 @@ namespace DataProviderListView
 	const FName HeaderIdName_GameThreadTiming = TEXT("GameThreadTiming");
 	const FName HeaderIdName_RenderThreadTiming = TEXT("RenderThreadTiming");
 	const FName HeaderIdName_GPUTiming = TEXT("GPUTiming");
+	const FName HeaderIdName_ShadersToCompile = TEXT("ShadersToCompile");
 }
 
 /**
@@ -58,7 +60,7 @@ struct FDataProviderTableRowData : TSharedFromThis<FDataProviderTableRowData>
 	}
 
 public:
-	
+
 	/** Identifier and descriptor associated to this list entry */
 	FGuid Identifier;
 	FStageInstanceDescriptor Descriptor;
@@ -96,6 +98,19 @@ TSharedRef<SWidget> SDataProviderTableRow::GenerateWidgetForColumn(const FName&
 			.Text(this, &SDataProviderTableRow::GetStateGlyphs)
 			.ColorAndOpacity(this, &SDataProviderTableRow::GetStateColorAndOpacity);
 	}
+	if (DataProviderListView::HeaderIdName_Status == ColumnName)
+	{
+		return SNew(SHorizontalBox)
+			+ SHorizontalBox::Slot()
+			.HAlign(HAlign_Left)
+			.VAlign(VAlign_Center)
+			.AutoWidth()
+			[
+				SNew(STextBlock)
+				.Text(this, &SDataProviderTableRow::GetStatus)
+				.ToolTipText(this, &SDataProviderTableRow::GetStatusToolTip)
+			];
+	}
 	if (DataProviderListView::HeaderIdName_Timecode == ColumnName)
 	{
 		return SNew(SHorizontalBox)
@@ -216,17 +231,64 @@ TSharedRef<SWidget> SDataProviderTableRow::GenerateWidgetForColumn(const FName&
 				.Text(MakeAttributeSP(this, &SDataProviderTableRow::GetGPUTiming))
 			];
 	}
+	if (DataProviderListView::HeaderIdName_ShadersToCompile == ColumnName)
+	{
+		return SNew(SHorizontalBox)
+			+ SHorizontalBox::Slot()
+			.HAlign(HAlign_Left)
+			.VAlign(VAlign_Center)
+			.AutoWidth()
+			[
+				SNew(STextBlock)
+				.Text(MakeAttributeSP(this, &SDataProviderTableRow::GetShadersLeftToCompile))
+			];
+	}
 
 	return SNullWidget::NullWidget;
 }
 
 FText SDataProviderTableRow::GetStateGlyphs() const
 {
+	if (Item->CachedPerformanceData.Status !=  EStageMonitorNodeStatus::Ready)
+	{
+		return FEditorFontGlyphs::Exclamation_Triangle;
+	}
 	return FEditorFontGlyphs::Circle;
 }
 
+FText SDataProviderTableRow::GetStatusToolTip() const
+{
+	return FText::AsCultureInvariant(Item->CachedPerformanceData.AssetInStatus);
+}
+
+FText SDataProviderTableRow::GetStatus() const
+{
+	switch(Item->CachedPerformanceData.Status)
+	{
+	case EStageMonitorNodeStatus::LoadingMap:
+		return LOCTEXT("Status_LoadingMap", "Loading");
+	case EStageMonitorNodeStatus::Ready:
+		return LOCTEXT("Status_Ready", "Ready");
+	case EStageMonitorNodeStatus::HotReload:
+		return LOCTEXT("Status_HotReload", "Reloading");
+	case EStageMonitorNodeStatus::ShaderCompiling:
+		return LOCTEXT("Status_ShaderCompiling", "Compiling");
+	case EStageMonitorNodeStatus::Unknown:
+		return LOCTEXT("Status_Unknown", "Unknown");
+	default:
+		break;
+	};
+
+	return LOCTEXT("Status_Undefined", "Undefined");
+}
+
 FSlateColor SDataProviderTableRow::GetStateColorAndOpacity() const
 {
+	if (Item->CachedState != EStageDataProviderState::Closed && Item->CachedPerformanceData.Status !=  EStageMonitorNodeStatus::Ready)
+	{
+		return FLinearColor::Yellow;
+	}
+
 	switch (Item->CachedState)
 	{
 		case EStageDataProviderState::Active:
@@ -296,6 +358,10 @@ FText SDataProviderTableRow::GetGPUTiming() const
 	return FText::AsNumber(Item->CachedPerformanceData.GPU_MS);
 }
 
+FText SDataProviderTableRow::GetShadersLeftToCompile() const
+{
+	return FText::AsNumber(Item->CachedPerformanceData.ShadersToCompile);
+}
 /**
  * SDataProviderListView
  */
@@ -325,6 +391,14 @@ void SDataProviderListView::Construct(const FArguments& InArgs, const TWeakPtr<I
 			.SortMode(this, &SDataProviderListView::GetColumnSortMode, DataProviderListView::HeaderIdName_State)
 			.OnSort(this, &SDataProviderListView::OnColumnSortModeChanged)
 
+			+ SHeaderRow::Column(DataProviderListView::HeaderIdName_Status)
+			.FixedWidth(55.f)
+			.HAlignHeader(HAlign_Center)
+			.HAlignCell(HAlign_Center)
+			.DefaultLabel(LOCTEXT("HeaderName_Status", "Status"))
+			.SortMode(this, &SDataProviderListView::GetColumnSortMode, DataProviderListView::HeaderIdName_Status)
+			.OnSort(this, &SDataProviderListView::OnColumnSortModeChanged)
+
 			+ SHeaderRow::Column(DataProviderListView::HeaderIdName_Timecode)
 			.FillWidth(.2f)
 			.DefaultLabel(LOCTEXT("HeaderName_Timecode", "Timecode"))
@@ -373,6 +447,10 @@ void SDataProviderListView::Construct(const FArguments& InArgs, const TWeakPtr<I
 			+ SHeaderRow::Column(DataProviderListView::HeaderIdName_GPUTiming)
 			.FillWidth(.2f)
 			.DefaultLabel(LOCTEXT("HeaderName_GPU", "GPU (ms)"))
+
+			+ SHeaderRow::Column(DataProviderListView::HeaderIdName_ShadersToCompile)
+			.FillWidth(.2f)
+			.DefaultLabel(LOCTEXT("HeaderName_Shaders", "Shaders To Compile"))
 		)
 	);
 
@@ -499,6 +577,11 @@ void SDataProviderListView::SortProviderList()
 		{
 			return CurrentSortMode == EColumnSortMode::Ascending ? Lhs->CachedState < Rhs->CachedState : Lhs->CachedState > Rhs->CachedState;
 		}
+		if (ColumnName == DataProviderListView::HeaderIdName_Status)
+		{
+			const int32 CompareResult = Lhs->CachedPerformanceData.Status < Rhs->CachedPerformanceData.Status;
+			return CurrentSortMode == EColumnSortMode::Ascending ? CompareResult < 0 : CompareResult > 0;
+		}
 		else if (ColumnName == DataProviderListView::HeaderIdName_MachineName)
 		{
 			const int32 CompareResult = Lhs->Descriptor.MachineName.Compare(Rhs->Descriptor.MachineName);
diff --git a/Engine/Plugins/VirtualProduction/StageMonitoring/Source/StageMonitorEditor/Private/Widgets/SDataProviderListView.h b/Engine/Plugins/VirtualProduction/StageMonitoring/Source/StageMonitorEditor/Private/Widgets/SDataProviderListView.h
--- a/Engine/Plugins/VirtualProduction/StageMonitoring/Source/StageMonitorEditor/Private/Widgets/SDataProviderListView.h
+++ b/Engine/Plugins/VirtualProduction/StageMonitoring/Source/StageMonitorEditor/Private/Widgets/SDataProviderListView.h
@@ -17,7 +17,7 @@ using FDataProviderTableRowDataPtr = TSharedPtr<FDataProviderTableRowData>;
 
 
 /**
- *	
+ *
  */
 class SDataProviderTableRow : public SMultiColumnTableRow<FDataProviderTableRowDataPtr>
 {
@@ -38,6 +38,8 @@ private:
 	/** Getters to populate the UI */
 	FText GetStateGlyphs() const;
 	FSlateColor GetStateColorAndOpacity() const;
+	FText GetStatus() const;
+	FText GetStatusToolTip() const;
 	FText GetTimecode() const;
 	FText GetMachineName() const;
 	FText GetProcessId() const;
@@ -48,11 +50,13 @@ private:
 	FText GetGameThreadTiming() const;
 	FText GetRenderThreadTiming() const;
 	FText GetGPUTiming() const;
+	FText GetShadersLeftToCompile() const;
+
 private:
 
 	/** Item to display */
 	FDataProviderTableRowDataPtr Item;
-	
+
 	/** Last time we refreshed the UI */
 	double LastRefreshTime = 0.0;
 };
@@ -84,10 +88,10 @@ private:
 
 	/** Generate a new row for the listview using the Item data */
 	TSharedRef<ITableRow> OnGenerateRow(FDataProviderTableRowDataPtr Item, const TSharedRef<STableViewBase>& OwnerTable);
-	
+
 	/** Callback triggered when provider list changed */
 	void OnStageMonitoringMachineListChanged();
-	
+
 	/** Rebuild provider list from scratch */
 	void RebuildDataProviderList();
 
@@ -104,7 +108,7 @@ private:
 	void SortProviderList();
 
 private:
-	
+
 	/** Widget and data containing info about what's shown in the listview */
 	TArray<FDataProviderTableRowDataPtr> ListItemsSource;
 	TArray<TWeakPtr<SDataProviderTableRow>> ListRowWidgets;
@@ -120,7 +124,7 @@ private:
 
 	/** Currently sorted column name */
 	FName SortedColumnName;
-	
+
 	/** Sort mode for the currently sorted column */
 	EColumnSortMode::Type SortMode = EColumnSortMode::Ascending;
 
diff --git a/Engine/Source/Runtime/VirtualProduction/StageDataCore/Private/StageMessages.cpp b/Engine/Source/Runtime/VirtualProduction/StageDataCore/Private/StageMessages.cpp
--- a/Engine/Source/Runtime/VirtualProduction/StageDataCore/Private/StageMessages.cpp
+++ b/Engine/Source/Runtime/VirtualProduction/StageDataCore/Private/StageMessages.cpp
@@ -40,3 +40,19 @@ FString FCriticalStateProviderMessage::ToString() const
 	}
 }
 
+FString FAssetLoadingStateProviderMessage::ToString() const
+{
+	switch (LoadingState)
+	{
+		case EStageLoadingState::PreLoad:
+		{
+			return FString::Printf(TEXT("Started loading asset: %s"), *AssetName);
+		}
+		case EStageLoadingState::PostLoad:
+		default:
+		{
+			return FString::Printf(TEXT("Finished loading asset: %s"), *AssetName);
+		}
+	}
+}
+
diff --git a/Engine/Source/Runtime/VirtualProduction/StageDataCore/Public/StageMessages.h b/Engine/Source/Runtime/VirtualProduction/StageDataCore/Public/StageMessages.h
--- a/Engine/Source/Runtime/VirtualProduction/StageDataCore/Public/StageMessages.h
+++ b/Engine/Source/Runtime/VirtualProduction/StageDataCore/Public/StageMessages.h
@@ -53,6 +53,16 @@ enum class EStageCriticalStateEvent : uint8
 	Exit,
 };
 
+/** Different events associated with stage critical state */
+UENUM()
+enum class EStageLoadingState : uint8
+{
+	/** Asset loading has started. */
+	PreLoad,
+	/** Asset loading has finished. */
+	PostLoad
+};
+
 /**
  * Holds descriptive information about that data providers. 
  * Information that won't change for a session
@@ -253,3 +263,29 @@ public:
 	UPROPERTY(VisibleAnywhere, Category = "CriticalState")
 	FName SourceName;
 };
+
+/**
+ * Message sent to indicate that the node has entered or exited a loading state.
+ */
+USTRUCT()
+struct STAGEDATACORE_API FAssetLoadingStateProviderMessage : public FStageProviderEventMessage
+{
+	GENERATED_BODY()
+
+public:
+	FAssetLoadingStateProviderMessage() = default;
+
+	FAssetLoadingStateProviderMessage(EStageLoadingState InState, const FString& InAssetName)
+		: LoadingState(InState), AssetName(InAssetName)
+	{}
+
+	virtual FString ToString() const override;
+
+	/** Event for this critical state */
+	UPROPERTY(VisibleAnywhere, Category = "AssetLoading")
+	EStageLoadingState LoadingState;
+
+	/** Name of the asset currently loading. */
+	UPROPERTY(VisibleAnywhere, Category = "AssetLoading")
+	FString AssetName;
+};
