Check to make sure the property change event is a value set. We are only interested in those events.

diff --git a/Engine/Plugins/VirtualProduction/StageMonitoring/Source/StageDataProvider/Private/FramePerformanceProvider.cpp b/Engine/Plugins/VirtualProduction/StageMonitoring/Source/StageDataProvider/Private/FramePerformanceProvider.cpp
--- a/Engine/Plugins/VirtualProduction/StageMonitoring/Source/StageDataProvider/Private/FramePerformanceProvider.cpp
+++ b/Engine/Plugins/VirtualProduction/StageMonitoring/Source/StageDataProvider/Private/FramePerformanceProvider.cpp
@@ -263,18 +263,22 @@ void FFramePerformanceProvider::CheckHitches(int64 Frame)
 #if WITH_EDITOR
 void FFramePerformanceProvider::OnStageSettingsChanged(UObject* Object, struct FPropertyChangedEvent& PropertyChangedEvent)
 {
-	const FName PropertyName = (PropertyChangedEvent.Property != nullptr) ? PropertyChangedEvent.Property->GetFName() : NAME_None;
-
-	if (PropertyName == GET_MEMBER_NAME_CHECKED(FStageHitchDetectionSettings, bEnableHitchDetection))
-	{
-		EnableHitchDetection(GetDefault<UStageMonitoringSettings>()->ProviderSettings.HitchDetectionSettings.bEnableHitchDetection);
-	}
-	if (PropertyName == GET_MEMBER_NAME_CHECKED(FStageFramePerformanceSettings, UpdateInterval))
+	if (PropertyChangedEvent.ChangeType == EPropertyChangeType::ValueSet)
 	{
-		const double UpdateInterval = GetDefault<UStageMonitoringSettings>()->ProviderSettings.FramePerformanceSettings.UpdateInterval;
-		ProviderThread->SetUpdateFrequency(static_cast<float>(UpdateInterval));
+		const FName PropertyName = (PropertyChangedEvent.Property != nullptr) ? PropertyChangedEvent.Property->GetFName() : NAME_None;
+
+		if (PropertyName == GET_MEMBER_NAME_CHECKED(FStageHitchDetectionSettings, bEnableHitchDetection))
+		{
+			EnableHitchDetection(GetDefault<UStageMonitoringSettings>()->ProviderSettings.HitchDetectionSettings.bEnableHitchDetection);
+		}
+		if (PropertyName == GET_MEMBER_NAME_CHECKED(FStageFramePerformanceSettings, UpdateInterval))
+		{
+			const double UpdateInterval = GetDefault<UStageMonitoringSettings>()->ProviderSettings.FramePerformanceSettings.UpdateInterval;
+			ProviderThread->SetUpdateFrequency(static_cast<float>(UpdateInterval));
+		}
 	}
 }
+
 #endif //WITH_EDITOR
 
 void FFramePerformanceProvider::EnableHitchDetection(bool bShouldEnable)
