Fix VP Roles subsystem crash

diff --git a/Engine/Plugins/Experimental/VirtualCamera/Source/VCamCore/Private/VCamComponent.cpp b/Engine/Plugins/Experimental/VirtualCamera/Source/VCamCore/Private/VCamComponent.cpp
--- a/Engine/Plugins/Experimental/VirtualCamera/Source/VCamCore/Private/VCamComponent.cpp
+++ b/Engine/Plugins/Experimental/VirtualCamera/Source/VCamCore/Private/VCamComponent.cpp
@@ -1649,7 +1649,8 @@ bool UVCamComponent::IsCameraInVPRole() const
 	// We are in a valid camera role if the user has not assigned a role or the current VPSettings role matches the
 	// assigned role.
 	//
-	return !Role.IsValid() || GEngine->GetEngineSubsystem<UVirtualProductionRolesSubsystem>()->HasActiveRole(Role);
+	UVirtualProductionRolesSubsystem* VPRolesSubsystem = GEngine->GetEngineSubsystem<UVirtualProductionRolesSubsystem>();
+	return !Role.IsValid() || (VPRolesSubsystem && VPRolesSubsystem->HasActiveRole(Role));
 #else
 	return true;
 #endif
diff --git a/Engine/Plugins/VirtualProduction/StageMonitoring/Source/StageDataProvider/Private/StageDataProvider.cpp b/Engine/Plugins/VirtualProduction/StageMonitoring/Source/StageDataProvider/Private/StageDataProvider.cpp
--- a/Engine/Plugins/VirtualProduction/StageMonitoring/Source/StageDataProvider/Private/StageDataProvider.cpp
+++ b/Engine/Plugins/VirtualProduction/StageMonitoring/Source/StageDataProvider/Private/StageDataProvider.cpp
@@ -191,10 +191,13 @@ bool FStageDataProvider::IsMessageTypeExcluded(UScriptStruct* MessageType) const
 	{
 		if (GEngine)
 		{
-			const FGameplayTagContainer& CurrentRoles = GEngine->GetEngineSubsystem<UVirtualProductionRolesSubsystem>()->GetRolesContainer_Private();
-			if (!SupportedRoles->HasAny(CurrentRoles))
+			if (UVirtualProductionRolesSubsystem* VPRolesSubsytem = GEngine->GetEngineSubsystem<UVirtualProductionRolesSubsystem>())
 			{
-				return true;
+				const FGameplayTagContainer& CurrentRoles = VPRolesSubsytem->GetRolesContainer_Private();
+				if (!SupportedRoles->HasAny(CurrentRoles))
+				{
+					return true;
+				}
 			}
 		}
 	}
diff --git a/Engine/Plugins/VirtualProduction/StageMonitoring/Source/StageMonitorCommon/Private/StageMonitorUtils.cpp b/Engine/Plugins/VirtualProduction/StageMonitoring/Source/StageMonitorCommon/Private/StageMonitorUtils.cpp
--- a/Engine/Plugins/VirtualProduction/StageMonitoring/Source/StageMonitorCommon/Private/StageMonitorUtils.cpp
+++ b/Engine/Plugins/VirtualProduction/StageMonitoring/Source/StageMonitorCommon/Private/StageMonitorUtils.cpp
@@ -22,7 +22,10 @@ namespace StageMonitorUtils
 
 		if (GEngine)
 		{
-			Descriptor.RolesStringified = GEngine->GetEngineSubsystem<UVirtualProductionRolesSubsystem>()->GetActiveRolesString();
+			if (const UVirtualProductionRolesSubsystem* VPRolesSubsystem = GEngine->GetEngineSubsystem<UVirtualProductionRolesSubsystem>())
+			{
+				Descriptor.RolesStringified = VPRolesSubsystem->GetActiveRolesString();
+			}
 		}
 
 		Descriptor.FriendlyName = GetDefault<UStageMonitoringSettings>()->CommandLineFriendlyName;
