Add GPU Mem / CPU Mem and all asset compilation tasks to Stage Monitor

diff --git a/Engine/Plugins/VirtualProduction/StageMonitoring/Config/BaseStageMonitoring.ini b/Engine/Plugins/VirtualProduction/StageMonitoring/Config/BaseStageMonitoring.ini
--- a/Engine/Plugins/VirtualProduction/StageMonitoring/Config/BaseStageMonitoring.ini
+++ b/Engine/Plugins/VirtualProduction/StageMonitoring/Config/BaseStageMonitoring.ini
@@ -6,3 +6,7 @@
 
 ; HitchDetection Settings
 +PropertyRedirects=(OldName="StageHitchDetectionSettings.TargetFrameRate",NewName="StageHitchDetectionSettings.MinimumFrameRate")
+
+; 5.3 Redirects
+;
++PropertyRedirects=(OldName="FramePerformanceProviderMessage.ShadersToCompile",NewName="FramePerformanceProviderMessage.CompilationTasksRemaining")
diff --git a/Engine/Plugins/VirtualProduction/StageMonitoring/Source/StageDataProvider/Private/FramePerformanceProvider.cpp b/Engine/Plugins/VirtualProduction/StageMonitoring/Source/StageDataProvider/Private/FramePerformanceProvider.cpp
--- a/Engine/Plugins/VirtualProduction/StageMonitoring/Source/StageDataProvider/Private/FramePerformanceProvider.cpp
+++ b/Engine/Plugins/VirtualProduction/StageMonitoring/Source/StageDataProvider/Private/FramePerformanceProvider.cpp
@@ -2,42 +2,61 @@
 
 #include "FramePerformanceProvider.h"
 
+#include "AssetCompilingManager.h"
 #include "EngineGlobals.h"
+#include "HAL/Runnable.h"
+#include "HAL/RunnableThread.h"
 #include "IStageDataProvider.h"
 #include "Misc/App.h"
+#include "RHI.h"
+#include "RenderTimer.h"
+
 #include "StageDataProviderModule.h"
 #include "StageMonitorUtils.h"
 #include "Stats/StatsData.h"
 #include "UObject/Package.h"
 #include "UObject/PackageReload.h"
 
-#include "ShaderCompiler.h"
-
 namespace UE::FramePerformanceProvider::Private
 {
 
-int32 ShadersLeftToCompile()
+int32 CompilationTasksRemaining()
+{
+	FAssetCompilingManager& Manager = FAssetCompilingManager::Get();
+	return Manager.GetNumRemainingAssets();
+}
+
+uint64 TotalGPUResourceSize(const TArray<FStatMessage>& StatMessages)
 {
-	if (GShaderCompilingManager && GShaderCompilingManager->IsCompiling())
+	uint64 TotalSize = 0;
+	FName NAME_STATGROUP_RHI(FStatGroup_STATGROUP_RHI::GetGroupName());
+	for (const FStatMessage& Stat : StatMessages)
 	{
-		return GShaderCompilingManager->GetNumRemainingJobs();
+		FName LastGroup = Stat.NameAndInfo.GetGroupName();
+		if (LastGroup == NAME_STATGROUP_RHI && Stat.NameAndInfo.GetFlag(EStatMetaFlags::IsMemory))
+		{
+			TotalSize += Stat.GetValue_int64();
+		}
 	}
-	return 0;
+	return TotalSize;
 }
 
-FFramePerformanceProviderMessage GetLatestPerformanceData()
+FFramePerformanceProviderMessage GetLatestPerformanceData(const	TArray<FStatMessage>& StatMessages)
 {
 	const float GameThreadTime = FPlatformTime::ToMilliseconds(GGameThreadTime);
 	const float RenderThreadTime = FPlatformTime::ToMilliseconds(GRenderThreadTime);
 	const float GPUTime = FPlatformTime::ToMilliseconds(GGPUFrameTime);
 	const float IdleTimeMilli = (FApp::GetIdleTime() * 1000.0);
+	FPlatformMemoryStats Stats = FPlatformMemory::GetStats();
 	return FFramePerformanceProviderMessage(
 		EStageMonitorNodeStatus::Unknown,
 		GameThreadTime,
 		RenderThreadTime,
 		GPUTime,
 		IdleTimeMilli,
-		ShadersLeftToCompile()
+		Stats.UsedPhysical,
+		TotalGPUResourceSize(StatMessages),
+		CompilationTasksRemaining()
 		);
 }
 
@@ -62,12 +81,23 @@ public:
 	// FRunnable functions
 	virtual uint32 Run() override
 	{
+		double LastRHIUpdate = FPlatformTime::Seconds();
+
 		bStopped = false;
 		do
 		{
+			const double CurrentPlatformTimeInSeconds = FPlatformTime::Seconds();
+			if (CurrentPlatformTimeInSeconds > (UpdateRHIResourcesFrequency + LastRHIUpdate)
+				|| Stats.IsEmpty())
+			{
+				Stats.Reset();
+				LastRHIUpdate = CurrentPlatformTimeInSeconds;
+				GetPermanentStats(Stats);
+			}
+
 			IStageDataProvider::SendMessage<FFramePerformanceProviderMessage>(EStageMessageFlags::None, GetFramePerformanceData());
 
-			FPlatformProcess::Sleep( UpdateFrequency );
+			FPlatformProcess::Sleep(UpdateFrequency);
 		} while (!bStopped);
 		return 0;
 	}
@@ -92,10 +122,11 @@ public:
 
 	FFramePerformanceProviderMessage GetFramePerformanceData()
 	{
-		FFramePerformanceProviderMessage OutboundData = UE::FramePerformanceProvider::Private::GetLatestPerformanceData();
-		if (OutboundData.ShadersToCompile > 0)
+		FFramePerformanceProviderMessage OutboundData =
+			UE::FramePerformanceProvider::Private::GetLatestPerformanceData(Stats);
+		if (OutboundData.CompilationTasksRemaining > 0)
 		{
-			OutboundData.Status = EStageMonitorNodeStatus::ShaderCompiling;
+			OutboundData.Status = EStageMonitorNodeStatus::AssetCompiling;
 		}
 		else
 		{
@@ -151,7 +182,9 @@ public:
 	mutable FCriticalSection LoadInfoCS;
 
 	TArray<FLoadInfo> LoadStack;
+	TArray<FStatMessage> Stats;
 
+	const float UpdateRHIResourcesFrequency = 4.0f;
 	float UpdateFrequency = 0.2f; // Default is 200ms;
 	bool bStopped = false;
 };
diff --git a/Engine/Plugins/VirtualProduction/StageMonitoring/Source/StageDataProvider/StageDataProvider.build.cs b/Engine/Plugins/VirtualProduction/StageMonitoring/Source/StageDataProvider/StageDataProvider.build.cs
--- a/Engine/Plugins/VirtualProduction/StageMonitoring/Source/StageDataProvider/StageDataProvider.build.cs
+++ b/Engine/Plugins/VirtualProduction/StageMonitoring/Source/StageDataProvider/StageDataProvider.build.cs
@@ -21,6 +21,7 @@ public class StageDataProvider : ModuleRules
 				"CoreUObject",
 				"DeveloperSettings",
 				"Engine",
+				"RHI",
 				"GameplayTags",
 				"RenderCore",
 				"StageDataCore",
diff --git a/Engine/Plugins/VirtualProduction/StageMonitoring/Source/StageMonitorCommon/Public/StageMonitorUtils.h b/Engine/Plugins/VirtualProduction/StageMonitoring/Source/StageMonitorCommon/Public/StageMonitorUtils.h
--- a/Engine/Plugins/VirtualProduction/StageMonitoring/Source/StageMonitorCommon/Public/StageMonitorUtils.h
+++ b/Engine/Plugins/VirtualProduction/StageMonitoring/Source/StageMonitorCommon/Public/StageMonitorUtils.h
@@ -15,7 +15,7 @@ enum class EStageMonitorNodeStatus
 	LoadingMap,
 	Ready,
 	HotReload,
-	ShaderCompiling
+	AssetCompiling
 };
 
 
@@ -32,16 +32,15 @@ public:
 
 	FFramePerformanceProviderMessage() = default;
 
-	UE_DEPRECATED(5.1, "FFramePerformanceProviderMessage constructor is deprecated, please use updated constructor.")
-	FFramePerformanceProviderMessage(float GameThreadTime, float RenderThreadTime, float GPUTime, float IdleTime)
-		: GameThreadMS(GameThreadTime), RenderThreadMS(RenderThreadTime), GPU_MS(GPUTime), IdleTimeMS(IdleTime), ShadersToCompile(0)
-	{
-		extern ENGINE_API float GAverageFPS;
-		AverageFPS = GAverageFPS;
-	}
-
-	FFramePerformanceProviderMessage(EStageMonitorNodeStatus InStatus, float GameThreadTime, float RenderThreadTime, float GPUTime, float IdleTime, int32 InShadersToCompile)
-	: Status(InStatus), GameThreadMS(GameThreadTime), RenderThreadMS(RenderThreadTime), GPU_MS(GPUTime), IdleTimeMS(IdleTime), ShadersToCompile(InShadersToCompile)
+	FFramePerformanceProviderMessage(EStageMonitorNodeStatus InStatus, float GameThreadTime, float RenderThreadTime, float GPUTime, float IdleTime, uint64 CPUMem, uint64 GPUMem, int32 InAssetsToCompile)
+		: Status(InStatus)
+		, GameThreadMS(GameThreadTime)
+		, RenderThreadMS(RenderThreadTime)
+		, GPU_MS(GPUTime)
+		, IdleTimeMS(IdleTime)
+		, CPU_MEM(CPUMem)
+		, GPU_MEM(GPUMem)
+		, CompilationTasksRemaining(InAssetsToCompile)
 	{
 		extern ENGINE_API float GAverageFPS;
 		AverageFPS = GAverageFPS;
@@ -75,9 +74,17 @@ public:
 	UPROPERTY(VisibleAnywhere, Category = "Performance", meta = (Unit = "ms"))
 	float IdleTimeMS = 0.f;
 
-	/** Number of shaders currently being compiled. */
+	/** Current CPU Memory Usage (Physical) in bytes */
+	UPROPERTY(VisibleAnywhere, Category = "Performance")
+	uint64 CPU_MEM = 0;
+
+	/** Current GPU Memory Usage (Physical) in bytes */
+	UPROPERTY(VisibleAnywhere, Category = "Performance")
+	uint64 GPU_MEM = 0;
+
+	/** Number of asynchronous compilation tasks currently in progress. */
 	UPROPERTY(VisibleAnywhere, Category = "Performance")
-	int32 ShadersToCompile = 0;
+	int32 CompilationTasksRemaining = 0;
 };
 
 
diff --git a/Engine/Plugins/VirtualProduction/StageMonitoring/Source/StageMonitorEditor/Private/Widgets/SDataProviderListView.cpp b/Engine/Plugins/VirtualProduction/StageMonitoring/Source/StageMonitorEditor/Private/Widgets/SDataProviderListView.cpp
--- a/Engine/Plugins/VirtualProduction/StageMonitoring/Source/StageMonitorEditor/Private/Widgets/SDataProviderListView.cpp
+++ b/Engine/Plugins/VirtualProduction/StageMonitoring/Source/StageMonitorEditor/Private/Widgets/SDataProviderListView.cpp
@@ -26,7 +26,9 @@ namespace DataProviderListView
 	const FName HeaderIdName_GameThreadTiming = TEXT("GameThreadTiming");
 	const FName HeaderIdName_RenderThreadTiming = TEXT("RenderThreadTiming");
 	const FName HeaderIdName_GPUTiming = TEXT("GPUTiming");
-	const FName HeaderIdName_ShadersToCompile = TEXT("ShadersToCompile");
+	const FName HeaderIdName_CPUMem = TEXT("CPUMem");
+	const FName HeaderIdName_GPUMem = TEXT("GPUMem");
+	const FName HeaderIdName_AssetsToCompile = TEXT("AssetsToCompile");
 }
 
 /**
@@ -231,7 +233,7 @@ TSharedRef<SWidget> SDataProviderTableRow::GenerateWidgetForColumn(const FName&
 				.Text(MakeAttributeSP(this, &SDataProviderTableRow::GetGPUTiming))
 			];
 	}
-	if (DataProviderListView::HeaderIdName_ShadersToCompile == ColumnName)
+	if (DataProviderListView::HeaderIdName_CPUMem == ColumnName)
 	{
 		return SNew(SHorizontalBox)
 			+ SHorizontalBox::Slot()
@@ -240,7 +242,31 @@ TSharedRef<SWidget> SDataProviderTableRow::GenerateWidgetForColumn(const FName&
 			.AutoWidth()
 			[
 				SNew(STextBlock)
-				.Text(MakeAttributeSP(this, &SDataProviderTableRow::GetShadersLeftToCompile))
+				.Text(MakeAttributeSP(this, &SDataProviderTableRow::GetCPUMem))
+			];
+	}
+	if (DataProviderListView::HeaderIdName_GPUMem == ColumnName)
+	{
+		return SNew(SHorizontalBox)
+			+ SHorizontalBox::Slot()
+			.HAlign(HAlign_Left)
+			.VAlign(VAlign_Center)
+			.AutoWidth()
+			[
+				SNew(STextBlock)
+				.Text(MakeAttributeSP(this, &SDataProviderTableRow::GetGPUMem))
+			];
+	}
+	if (DataProviderListView::HeaderIdName_AssetsToCompile == ColumnName)
+	{
+		return SNew(SHorizontalBox)
+			+ SHorizontalBox::Slot()
+			.HAlign(HAlign_Left)
+			.VAlign(VAlign_Center)
+			.AutoWidth()
+			[
+				SNew(STextBlock)
+				.Text(MakeAttributeSP(this, &SDataProviderTableRow::GetAssetsLeftToCompile))
 			];
 	}
 
@@ -271,8 +297,8 @@ FText SDataProviderTableRow::GetStatus() const
 		return LOCTEXT("Status_Ready", "Ready");
 	case EStageMonitorNodeStatus::HotReload:
 		return LOCTEXT("Status_HotReload", "Reloading");
-	case EStageMonitorNodeStatus::ShaderCompiling:
-		return LOCTEXT("Status_ShaderCompiling", "Compiling");
+	case EStageMonitorNodeStatus::AssetCompiling:
+		return LOCTEXT("Status_AssetCompiling", "Asset Compiling");
 	case EStageMonitorNodeStatus::Unknown:
 		return LOCTEXT("Status_Unknown", "Unknown");
 	default:
@@ -307,6 +333,22 @@ FSlateColor SDataProviderTableRow::GetStateColorAndOpacity() const
 	}
 }
 
+FString GetMemoryString( const uint64 Value )
+{
+	if (Value > 1024 * 1024)
+	{
+		// Display anything over a MB in GB format.
+		return FString::Printf( TEXT( "%.2f GB" ), float( static_cast<float>(Value) / (1024.0 * 1024.0 * 1024.0) ) );
+	}
+	else if (Value > 1024)
+	{
+		// Anything more than a KB in MB.
+		return FString::Printf( TEXT( "%.2f MB" ), float( static_cast<float>(Value) / (1024.0 * 1024.0) ) );
+	}
+
+	return FString::Printf( TEXT( "%.2f KB" ), float( static_cast<float>(Value) / (1024.0) ) );
+}
+
 FText SDataProviderTableRow::GetTimecode() const
 {
 	const FTimecode CachedTimecode = FTimecode::FromFrameNumber(Item->CachedPerformanceData.FrameTime.Time.FrameNumber, Item->CachedPerformanceData.FrameTime.Rate);
@@ -358,9 +400,19 @@ FText SDataProviderTableRow::GetGPUTiming() const
 	return FText::AsNumber(Item->CachedPerformanceData.GPU_MS);
 }
 
-FText SDataProviderTableRow::GetShadersLeftToCompile() const
+FText SDataProviderTableRow::GetGPUMem() const
 {
-	return FText::AsNumber(Item->CachedPerformanceData.ShadersToCompile);
+	return FText::FromString(GetMemoryString(Item->CachedPerformanceData.GPU_MEM));
+}
+
+FText SDataProviderTableRow::GetCPUMem() const
+{
+	return FText::FromString(GetMemoryString(Item->CachedPerformanceData.CPU_MEM));
+}
+
+FText SDataProviderTableRow::GetAssetsLeftToCompile() const
+{
+	return FText::AsNumber(Item->CachedPerformanceData.CompilationTasksRemaining);
 }
 /**
  * SDataProviderListView
@@ -448,9 +500,17 @@ void SDataProviderListView::Construct(const FArguments& InArgs, const TWeakPtr<I
 			.FillWidth(.2f)
 			.DefaultLabel(LOCTEXT("HeaderName_GPU", "GPU (ms)"))
 
-			+ SHeaderRow::Column(DataProviderListView::HeaderIdName_ShadersToCompile)
+			+ SHeaderRow::Column(DataProviderListView::HeaderIdName_CPUMem)
+			.FillWidth(.2f)
+			.DefaultLabel(LOCTEXT("HeaderName_CPUMem", "CPU Memory"))
+
+			+ SHeaderRow::Column(DataProviderListView::HeaderIdName_GPUMem)
+			.FillWidth(.2f)
+			.DefaultLabel(LOCTEXT("HeaderName_GPUMem", "GPU Memory"))
+
+			+ SHeaderRow::Column(DataProviderListView::HeaderIdName_AssetsToCompile)
 			.FillWidth(.2f)
-			.DefaultLabel(LOCTEXT("HeaderName_Shaders", "Shaders To Compile"))
+			.DefaultLabel(LOCTEXT("HeaderName_Assets", "Assets To Compile"))
 		)
 	);
 
diff --git a/Engine/Plugins/VirtualProduction/StageMonitoring/Source/StageMonitorEditor/Private/Widgets/SDataProviderListView.h b/Engine/Plugins/VirtualProduction/StageMonitoring/Source/StageMonitorEditor/Private/Widgets/SDataProviderListView.h
--- a/Engine/Plugins/VirtualProduction/StageMonitoring/Source/StageMonitorEditor/Private/Widgets/SDataProviderListView.h
+++ b/Engine/Plugins/VirtualProduction/StageMonitoring/Source/StageMonitorEditor/Private/Widgets/SDataProviderListView.h
@@ -46,7 +46,9 @@ private:
 	FText GetGameThreadTiming() const;
 	FText GetRenderThreadTiming() const;
 	FText GetGPUTiming() const;
-	FText GetShadersLeftToCompile() const;
+	FText GetGPUMem() const;
+	FText GetCPUMem() const;
+	FText GetAssetsLeftToCompile() const;
 
 private:
 
