Restructure how we handle columns in stage monitor and make them "saveable" across editor sessions

diff --git a/Engine/Plugins/VirtualProduction/StageMonitoring/Source/StageMonitorEditor/Private/StageMonitorEditorSettings.h b/Engine/Plugins/VirtualProduction/StageMonitoring/Source/StageMonitorEditor/Private/StageMonitorEditorSettings.h
--- a/Engine/Plugins/VirtualProduction/StageMonitoring/Source/StageMonitorEditor/Private/StageMonitorEditorSettings.h
+++ b/Engine/Plugins/VirtualProduction/StageMonitoring/Source/StageMonitorEditor/Private/StageMonitorEditorSettings.h
@@ -19,4 +19,8 @@ public:
 	/** Refresh rate in seconds for the StageMonitor panel. */
 	UPROPERTY(Config, EditAnywhere, Category = "UI", meta = (ClampMin = 0.0f))
 	float RefreshRate = 0.2f;
+
+	/** Stores the state of the visibility settings in JSON format. */
+	UPROPERTY(Config)
+	FString ColumnVisibilitySettings;
 };
diff --git a/Engine/Plugins/VirtualProduction/StageMonitoring/Source/StageMonitorEditor/Private/Widgets/SDataProviderListView.cpp b/Engine/Plugins/VirtualProduction/StageMonitoring/Source/StageMonitorEditor/Private/Widgets/SDataProviderListView.cpp
--- a/Engine/Plugins/VirtualProduction/StageMonitoring/Source/StageMonitorEditor/Private/Widgets/SDataProviderListView.cpp
+++ b/Engine/Plugins/VirtualProduction/StageMonitoring/Source/StageMonitorEditor/Private/Widgets/SDataProviderListView.cpp
@@ -2,35 +2,392 @@
 
 #include "SDataProviderListView.h"
 
+#include "Containers/StaticArray.h"
+#include "Dom/JsonObject.h"
 #include "EditorFontGlyphs.h"
 #include "IStageMonitorSession.h"
+#include "Layout/Visibility.h"
 #include "Misc/App.h"
-#include "StageMonitorUtils.h"
+#include "Misc/CoreMiscDefines.h"
+#include "Misc/ScopeExit.h"
+#include "Serialization/JsonSerializer.h"
+#include "Serialization/JsonWriter.h"
 #include "StageMonitorEditorSettings.h"
+#include "StageMonitorUtils.h"
+#include "Templates/UnrealTemplate.h"
 #include "UObject/StructOnScope.h"
+#include "Widgets/SDataProviderListView.h"
+#include "Widgets/Views/SHeaderRow.h"
+
 
 #define LOCTEXT_NAMESPACE "SDataProviderListView"
 
+/**
+ * This structure holds onto the definition and function callbacks for columns
+ * in the Stage Data Monitor list. Each definition specifies the following
+ *
+ * - The label (display name)
+ * - A flag to indicate if can be sorted
+ * - A flag indicate if it is default visible
+ * - A function callback when creating a new row
+ * - A function callback that can be used to add additional slate args.
+ * - An order id that specifies the order in the column is created.
+ */
+struct FColumnDefinition
+{
+	using FGenerateRowFuncType = TSharedRef<SWidget>(SDataProviderTableRow*);
+	using FSetColumnArgumentsFuncType = void(SHeaderRow::FColumn::FArguments&);
+
+	FColumnDefinition(TAttribute<FText> InLabel, bool bInSortable, bool bInIsVisible,
+					  FGenerateRowFuncType* InGenerateRowFuncPtr,
+					  FSetColumnArgumentsFuncType* InArgumentsColumnPtr = nullptr)
+		: Label(MoveTemp(InLabel))
+		, bSortable(bInSortable)
+		, bIsVisible(bInIsVisible)
+		, GenerateRowFuncPtr(InGenerateRowFuncPtr)
+		, ColumnArgumentsFuncPtr(InArgumentsColumnPtr)
+	{
+		if (bTrackOrderIdCount)
+		{
+			OrderId = OrderIdCounter;
+			OrderIdCounter++;
+		}
+	}
+
+	FColumnDefinition() = delete;
+
+	/** Label for the column */
+	TAttribute<FText> Label;
+
+	/** Can this column be sorted. */
+	bool  bSortable = false;
+
+	/** Is this column visible by default. */
+	bool  bIsVisible = true;
+
+	/** Function callback for generating a new row. */
+	FGenerateRowFuncType* GenerateRowFuncPtr = nullptr;
+
+	/** A function callback for column arguments. */
+	FSetColumnArgumentsFuncType* ColumnArgumentsFuncPtr = nullptr;
+
+	/** Procedurally generated order id.*/
+	uint32 OrderId = 0;
+
+	static void BeginOrderIdCount()
+	{
+		bTrackOrderIdCount = true;
+	}
+	static void EndOrderIdCount()
+	{
+		bTrackOrderIdCount = false;
+	}
+private:
+
+	static bool bTrackOrderIdCount;
+	static uint32 OrderIdCounter;
+};
+
+uint32 FColumnDefinition::OrderIdCounter = 0;
+bool FColumnDefinition::bTrackOrderIdCount = false;
 
 namespace DataProviderListView
 {
-	const FName HeaderIdName_State = TEXT("State");
-	const FName HeaderIdName_Status = TEXT("Status");
-	const FName HeaderIdName_Timecode = TEXT("Timecode");
-	const FName HeaderIdName_MachineName = TEXT("MachineName");
-	const FName HeaderIdName_ProcessId = TEXT("ProcessId");
-	const FName HeaderIdName_StageName = TEXT("StageName");
-	const FName HeaderIdName_Roles = TEXT("Roles");
-	const FName HeaderIdName_AverageFPS = TEXT("AverageFPS");
-	const FName HeaderIdName_IdleTime = TEXT("IdleTime");
-	const FName HeaderIdName_GameThreadTiming = TEXT("GameThreadTiming");
-	const FName HeaderIdName_RenderThreadTiming = TEXT("RenderThreadTiming");
-	const FName HeaderIdName_GPUTiming = TEXT("GPUTiming");
-	const FName HeaderIdName_CPUMem = TEXT("CPUMem");
-	const FName HeaderIdName_GPUMem = TEXT("GPUMem");
-	const FName HeaderIdName_AssetsToCompile = TEXT("AssetsToCompile");
+
+	const bool IsVisible = true;
+	const bool IsNotVisible = false;
+	const bool IsSortable = true;
+	const bool IsNotSortable = false;
+
+	namespace HeaderDef {
+
+	const FName State = TEXT("State");
+	const FName Status = TEXT("Status");
+	const FName Timecode = TEXT("Timecode");
+	const FName MachineName = TEXT("MachineName");
+	const FName ProcessId = TEXT("ProcessId");
+	const FName StageName = TEXT("StageName");
+	const FName Roles = TEXT("Roles");
+	const FName AverageFPS = TEXT("AverageFPS");
+	const FName IdleTime = TEXT("IdleTime");
+	const FName GameThreadTiming = TEXT("GameThreadTiming");
+	const FName RenderThreadTiming = TEXT("RenderThreadTiming");
+	const FName GPUTiming = TEXT("GPUTiming");
+	const FName CPUMem = TEXT("CPUMem");
+	const FName GPUMem = TEXT("GPUMem");
+	const FName AssetsToCompile = TEXT("AssetsToCompile");
+
+	const TMap<FName, FColumnDefinition> StaticColumnHeader = []() {
+		FColumnDefinition::BeginOrderIdCount();
+		TMap<FName, FColumnDefinition> Header = {
+			{State,
+			 {
+				 LOCTEXT("HeaderName_State", "State"), IsSortable, IsVisible,
+				 // Generate Column
+				 [](SDataProviderTableRow* Provider) -> TSharedRef<SWidget>
+				 {
+					 return SNew(SHorizontalBox)
+					 + SHorizontalBox::Slot()
+					 .HAlign(HAlign_Center)
+					 .VAlign(VAlign_Center)
+					 .AutoWidth()
+					 [
+						 SNew(STextBlock)
+							 .Font(FAppStyle::Get().GetFontStyle("FontAwesome.11"))
+							 .Text(Provider, &SDataProviderTableRow::GetStateGlyphs)
+							 .ColorAndOpacity(Provider, &SDataProviderTableRow::GetStateColorAndOpacity)
+					 ];
+				 },
+				 [](SHeaderRow::FColumn::FArguments& Args)
+				 {
+					 Args.FixedWidth(55.f);
+				 }
+			 }},
+			{Status,
+			 {
+				 LOCTEXT("HeaderName_Status", "Status"), IsSortable, IsVisible,
+				 // Generate Column
+				 [](SDataProviderTableRow* Provider) -> TSharedRef<SWidget>
+				 {
+					 return SNew(SHorizontalBox)
+					 + SHorizontalBox::Slot()
+					 .HAlign(HAlign_Left)
+					 .VAlign(VAlign_Center)
+					 .AutoWidth()
+					 [
+						 SNew(STextBlock)
+						 .Text(Provider, &SDataProviderTableRow::GetStatus)
+						 .ToolTipText(Provider, &SDataProviderTableRow::GetStatusToolTip)
+					 ];
+				 },
+				 [](SHeaderRow::FColumn::FArguments& Args)
+				 {
+					 Args.FixedWidth(55.f);
+				 }
+			 }},
+			{Timecode,
+			 {
+				 LOCTEXT("HeaderName_Timecode", "Timecode"), IsNotSortable, IsVisible,
+				 [](SDataProviderTableRow* Provider) -> TSharedRef<SWidget>
+				 {
+					 return SNew(SHorizontalBox)
+					 + SHorizontalBox::Slot()
+					 .HAlign(HAlign_Left)
+					 .VAlign(VAlign_Center)
+					 .AutoWidth()
+					 [
+						 SNew(STextBlock)
+						 .Text(Provider, &SDataProviderTableRow::GetTimecode)
+					 ];
+				 }
+			 }
+			},
+			{MachineName,
+			 {
+				 LOCTEXT("HeaderName_MachineName", "Machine"), IsNotSortable, IsVisible,
+				 [](SDataProviderTableRow* Provider) -> TSharedRef<SWidget>
+				 {
+					 return SNew(SHorizontalBox)
+					 + SHorizontalBox::Slot()
+					 .HAlign(HAlign_Left)
+					 .VAlign(VAlign_Center)
+					 .AutoWidth()
+					 [
+						 SNew(STextBlock)
+						 .Text(Provider, &SDataProviderTableRow::GetMachineName)
+					 ];
+				 }
+			 }
+			},
+			{ProcessId,
+			 {
+				 LOCTEXT("HeaderName_ProcessId", "Process Id"), IsSortable, IsNotVisible,
+				 [](SDataProviderTableRow* Provider) -> TSharedRef<SWidget>
+				 {
+					 return SNew(SHorizontalBox)
+					 + SHorizontalBox::Slot()
+					 .HAlign(HAlign_Left)
+					 .VAlign(VAlign_Center)
+					 .AutoWidth()
+					 [
+						 SNew(STextBlock)
+						 .Text(Provider, &SDataProviderTableRow::GetProcessId)
+					 ];
+				 }
+			 }},
+			{StageName,
+			 {
+				 LOCTEXT("HeaderName_StageName", "Stage Name"), IsSortable, IsNotVisible,
+				 [](SDataProviderTableRow* Provider) -> TSharedRef<SWidget>
+				 {
+					 return SNew(SHorizontalBox)
+					 + SHorizontalBox::Slot()
+					 .HAlign(HAlign_Left)
+					 .VAlign(VAlign_Center)
+					 .AutoWidth()
+					 [
+						 SNew(STextBlock)
+						 .Text(Provider, &SDataProviderTableRow::GetStageName)
+					 ];
+				 }
+			 }},
+			{Roles,
+			 {
+				 LOCTEXT("HeaderName_Roles", "Roles"), IsSortable, IsNotVisible,
+				 [](SDataProviderTableRow* Provider) -> TSharedRef<SWidget>
+				 {
+					 return SNew(SHorizontalBox)
+					 + SHorizontalBox::Slot()
+					 .HAlign(HAlign_Left)
+					 .VAlign(VAlign_Center)
+					 .AutoWidth()
+					 [
+						 SNew(STextBlock)
+						 .Text(Provider, &SDataProviderTableRow::GetRoles)
+					 ];
+				 }
+			 }},
+			{AverageFPS,
+			 {
+				 LOCTEXT("HeaderName_AverageFPS", "Average FPS"), IsNotSortable, IsVisible,
+				 [](SDataProviderTableRow* Provider) -> TSharedRef<SWidget>
+				 {
+					 return SNew(SHorizontalBox)
+					 + SHorizontalBox::Slot()
+					 .HAlign(HAlign_Left)
+					 .VAlign(VAlign_Center)
+					 .AutoWidth()
+					 [
+						 SNew(STextBlock)
+						 .Text(MakeAttributeSP(Provider, &SDataProviderTableRow::GetAverageFPS))
+					 ];
+				 }
+			 }},
+			{IdleTime,
+			 {
+				 LOCTEXT("HeaderName_IdleTime", "Idle Time (ms)"), IsNotSortable, IsNotVisible,
+				 [](SDataProviderTableRow* Provider) -> TSharedRef<SWidget>
+				 {
+					 return SNew(SHorizontalBox)
+					 + SHorizontalBox::Slot()
+					 .HAlign(HAlign_Left)
+					 .VAlign(VAlign_Center)
+					 .AutoWidth()
+					 [
+						 SNew(STextBlock)
+						 .Text(MakeAttributeSP(Provider, &SDataProviderTableRow::GetIdleTime))
+					 ];
+				 }
+			 }},
+			{GameThreadTiming,
+			 {
+				 LOCTEXT("HeaderName_GameThread", "Game Thread (ms)"), IsNotSortable, IsVisible,
+				 [](SDataProviderTableRow* Provider) -> TSharedRef<SWidget>
+				 {
+					 return SNew(SHorizontalBox)
+					 + SHorizontalBox::Slot()
+					 .HAlign(HAlign_Left)
+					 .VAlign(VAlign_Center)
+					 .AutoWidth()
+					 [
+						 SNew(STextBlock)
+						 .Text(MakeAttributeSP(Provider, &SDataProviderTableRow::GetGameThreadTiming))
+					 ];
+				 }
+			 }
+			},
+			{RenderThreadTiming,
+			 {
+				 LOCTEXT("HeaderName_RenderThread", "Render Thread (ms)"), IsNotSortable, IsVisible,
+				 [](SDataProviderTableRow* Provider) -> TSharedRef<SWidget>
+				 {
+					 return SNew(SHorizontalBox)
+					 + SHorizontalBox::Slot()
+					 .HAlign(HAlign_Left)
+					 .VAlign(VAlign_Center)
+					 .AutoWidth()
+					 [
+						 SNew(STextBlock)
+						 .Text(MakeAttributeSP(Provider, &SDataProviderTableRow::GetRenderThreadTiming))
+					 ];
+				 }
+			 }},
+			{GPUTiming,
+			 {
+				 LOCTEXT("HeaderName_GPU", "GPU (ms)"), IsNotSortable, IsVisible,
+				 [](SDataProviderTableRow* Provider) -> TSharedRef<SWidget>
+				 {
+					 return SNew(SHorizontalBox)
+					 + SHorizontalBox::Slot()
+					 .HAlign(HAlign_Left)
+					 .VAlign(VAlign_Center)
+					 .AutoWidth()
+					 [
+						 SNew(STextBlock)
+						 .Text(MakeAttributeSP(Provider, &SDataProviderTableRow::GetGPUTiming))
+					 ];
+				 }
+			 }},
+			{CPUMem,
+			 {
+				 LOCTEXT("HeaderName_CPUMem", "CPU Memory"), IsNotSortable, IsNotVisible,
+				 [](SDataProviderTableRow* Provider) -> TSharedRef<SWidget>
+				 {
+					 return SNew(SHorizontalBox)
+					 + SHorizontalBox::Slot()
+					 .HAlign(HAlign_Left)
+					 .VAlign(VAlign_Center)
+					 .AutoWidth()
+					 [
+						 SNew(STextBlock)
+						 .Text(MakeAttributeSP(Provider, &SDataProviderTableRow::GetCPUMem))
+						 ];
+				 }
+			 }},
+			{GPUMem,
+			 {
+				 LOCTEXT("HeaderName_GPUMem", "GPU Memory"), IsNotSortable, IsNotVisible,
+				 [](SDataProviderTableRow* Provider) -> TSharedRef<SWidget>
+				 {
+					 return SNew(SHorizontalBox)
+					 + SHorizontalBox::Slot()
+					 .HAlign(HAlign_Left)
+					 .VAlign(VAlign_Center)
+					 .AutoWidth()
+					 [
+						 SNew(STextBlock)
+						 .Text(MakeAttributeSP(Provider, &SDataProviderTableRow::GetGPUMem))
+					 ];
+				 }
+			 }
+			},
+			{AssetsToCompile,
+			 {
+				 LOCTEXT("HeaderName_Assets", "Assets To Compile"), IsNotSortable, IsNotVisible,
+				 [](SDataProviderTableRow* Provider) -> TSharedRef<SWidget>
+				 {
+					 return SNew(SHorizontalBox)
+					 + SHorizontalBox::Slot()
+					 .HAlign(HAlign_Left)
+					 .VAlign(VAlign_Center)
+					 .AutoWidth()
+					 [
+						 SNew(STextBlock)
+						 .Text(MakeAttributeSP(Provider, &SDataProviderTableRow::GetAssetsLeftToCompile))
+					 ];
+				 }
+			 }
+			}};
+		FColumnDefinition::EndOrderIdCount();
+		return Header;
+	}();
+
+	const FName FirstColumnId(State);
+
+	} // end namespace HeaderDef
 }
 
+
 /**
  * FDataProviderTableRowData
  */
@@ -67,7 +424,6 @@ public:
 	FGuid Identifier;
 	FStageInstanceDescriptor Descriptor;
 
-
 	/** Weak pointer to the session data */
 	TWeakPtr<IStageMonitorSession> Session;
 
@@ -88,188 +444,18 @@ void SDataProviderTableRow::Construct(const FArguments& InArgs, const TSharedRef
 	check(Item.IsValid());
 
 	Super::FArguments Arg;
+	Arg._Padding = InArgs._Padding;
 	Super::Construct(Arg, InOwerTableView);
 }
 
 TSharedRef<SWidget> SDataProviderTableRow::GenerateWidgetForColumn(const FName& ColumnName)
 {
-	if (DataProviderListView::HeaderIdName_State == ColumnName)
-	{
-		return SNew(STextBlock)
-			.Font(FAppStyle::Get().GetFontStyle("FontAwesome.11"))
-			.Text(this, &SDataProviderTableRow::GetStateGlyphs)
-			.ColorAndOpacity(this, &SDataProviderTableRow::GetStateColorAndOpacity);
-	}
-	if (DataProviderListView::HeaderIdName_Status == ColumnName)
-	{
-		return SNew(SHorizontalBox)
-			+ SHorizontalBox::Slot()
-			.HAlign(HAlign_Left)
-			.VAlign(VAlign_Center)
-			.AutoWidth()
-			[
-				SNew(STextBlock)
-				.Text(this, &SDataProviderTableRow::GetStatus)
-				.ToolTipText(this, &SDataProviderTableRow::GetStatusToolTip)
-			];
-	}
-	if (DataProviderListView::HeaderIdName_Timecode == ColumnName)
-	{
-		return SNew(SHorizontalBox)
-			+ SHorizontalBox::Slot()
-			.HAlign(HAlign_Left)
-			.VAlign(VAlign_Center)
-			.AutoWidth()
-			[
-				SNew(STextBlock)
-				.Text(this, &SDataProviderTableRow::GetTimecode)
-			];
-	}
-	if (DataProviderListView::HeaderIdName_MachineName == ColumnName)
-	{
-		return SNew(SHorizontalBox)
-			+ SHorizontalBox::Slot()
-			.HAlign(HAlign_Left)
-			.VAlign(VAlign_Center)
-			.AutoWidth()
-			[
-				SNew(STextBlock)
-				.Text(this, &SDataProviderTableRow::GetMachineName)
-			];
-	}
-	if (DataProviderListView::HeaderIdName_ProcessId == ColumnName)
-	{
-		return SNew(SHorizontalBox)
-			+ SHorizontalBox::Slot()
-			.HAlign(HAlign_Left)
-			.VAlign(VAlign_Center)
-			.AutoWidth()
-			[
-				SNew(STextBlock)
-				.Text(this, &SDataProviderTableRow::GetProcessId)
-			];
-	}
-	if (DataProviderListView::HeaderIdName_StageName == ColumnName)
-	{
-		return SNew(SHorizontalBox)
-			+ SHorizontalBox::Slot()
-			.HAlign(HAlign_Left)
-			.VAlign(VAlign_Center)
-			.AutoWidth()
-			[
-				SNew(STextBlock)
-				.Text(this, &SDataProviderTableRow::GetStageName)
-			];
-	}
-	if (DataProviderListView::HeaderIdName_Roles == ColumnName)
-	{
-		return SNew(SHorizontalBox)
-			+ SHorizontalBox::Slot()
-			.HAlign(HAlign_Left)
-			.VAlign(VAlign_Center)
-			.AutoWidth()
-			[
-				SNew(STextBlock)
-				.Text(this, &SDataProviderTableRow::GetRoles)
-			];
-	}
-	if (DataProviderListView::HeaderIdName_AverageFPS == ColumnName)
-	{
-		return SNew(SHorizontalBox)
-			+ SHorizontalBox::Slot()
-			.HAlign(HAlign_Left)
-			.VAlign(VAlign_Center)
-			.AutoWidth()
-			[
-				SNew(STextBlock)
-				.Text(MakeAttributeSP(this, &SDataProviderTableRow::GetAverageFPS))
-			];
-	}
-	if (DataProviderListView::HeaderIdName_IdleTime == ColumnName)
-	{
-		return SNew(SHorizontalBox)
-			+ SHorizontalBox::Slot()
-			.HAlign(HAlign_Left)
-			.VAlign(VAlign_Center)
-			.AutoWidth()
-			[
-				SNew(STextBlock)
-				.Text(MakeAttributeSP(this, &SDataProviderTableRow::GetIdleTime))
-			];
-	}
-	if (DataProviderListView::HeaderIdName_GameThreadTiming == ColumnName)
+	using namespace DataProviderListView::HeaderDef;
+	const FColumnDefinition* ColumnDef = StaticColumnHeader.Find(ColumnName);
+	if (ColumnDef)
 	{
-		return SNew(SHorizontalBox)
-			+ SHorizontalBox::Slot()
-			.HAlign(HAlign_Left)
-			.VAlign(VAlign_Center)
-			.AutoWidth()
-			[
-				SNew(STextBlock)
-				.Text(MakeAttributeSP(this, &SDataProviderTableRow::GetGameThreadTiming))
-			];
+		return ColumnDef->GenerateRowFuncPtr(this);
 	}
-	if (DataProviderListView::HeaderIdName_RenderThreadTiming == ColumnName)
-	{
-		return SNew(SHorizontalBox)
-			+ SHorizontalBox::Slot()
-			.HAlign(HAlign_Left)
-			.VAlign(VAlign_Center)
-			.AutoWidth()
-			[
-				SNew(STextBlock)
-				.Text(MakeAttributeSP(this, &SDataProviderTableRow::GetRenderThreadTiming))
-			];
-	}
-	if (DataProviderListView::HeaderIdName_GPUTiming == ColumnName)
-	{
-		return SNew(SHorizontalBox)
-			+ SHorizontalBox::Slot()
-			.HAlign(HAlign_Left)
-			.VAlign(VAlign_Center)
-			.AutoWidth()
-			[
-				SNew(STextBlock)
-				.Text(MakeAttributeSP(this, &SDataProviderTableRow::GetGPUTiming))
-			];
-	}
-	if (DataProviderListView::HeaderIdName_CPUMem == ColumnName)
-	{
-		return SNew(SHorizontalBox)
-			+ SHorizontalBox::Slot()
-			.HAlign(HAlign_Left)
-			.VAlign(VAlign_Center)
-			.AutoWidth()
-			[
-				SNew(STextBlock)
-				.Text(MakeAttributeSP(this, &SDataProviderTableRow::GetCPUMem))
-			];
-	}
-	if (DataProviderListView::HeaderIdName_GPUMem == ColumnName)
-	{
-		return SNew(SHorizontalBox)
-			+ SHorizontalBox::Slot()
-			.HAlign(HAlign_Left)
-			.VAlign(VAlign_Center)
-			.AutoWidth()
-			[
-				SNew(STextBlock)
-				.Text(MakeAttributeSP(this, &SDataProviderTableRow::GetGPUMem))
-			];
-	}
-	if (DataProviderListView::HeaderIdName_AssetsToCompile == ColumnName)
-	{
-		return SNew(SHorizontalBox)
-			+ SHorizontalBox::Slot()
-			.HAlign(HAlign_Left)
-			.VAlign(VAlign_Center)
-			.AutoWidth()
-			[
-				SNew(STextBlock)
-				.Text(MakeAttributeSP(this, &SDataProviderTableRow::GetAssetsLeftToCompile))
-			];
-	}
-
 	return SNullWidget::NullWidget;
 }
 
@@ -414,109 +600,168 @@ FText SDataProviderTableRow::GetAssetsLeftToCompile() const
 {
 	return FText::AsNumber(Item->CachedPerformanceData.CompilationTasksRemaining);
 }
+
+namespace UE::StageMonitor::Private
+{
+	namespace JsonKeys
+	{
+		const FString ColumnId("ColumnID");
+		const FString IsVisible("bIsVisible");
+		const FString Values("Values");
+	}
+	void CaptureColumnVisibilityState(TSharedRef<SHeaderRow> HeaderRow)
+	{
+		TSharedPtr<FJsonObject> JsonObject = MakeShared<FJsonObject>();
+		TArray<TSharedPtr<FJsonValue>> Array;
+		for (const SHeaderRow::FColumn& Column : HeaderRow->GetColumns())
+		{
+			TSharedPtr<FJsonObject> ColumnObject = MakeShared<FJsonObject>();
+			ColumnObject->SetStringField(JsonKeys::ColumnId, Column.ColumnId.ToString());
+			ColumnObject->SetBoolField(JsonKeys::IsVisible, Column.bIsVisible);
+			Array.Add(MakeShared<FJsonValueObject>(ColumnObject));
+		}
+		JsonObject->SetArrayField(JsonKeys::Values, Array);
+
+		FString Snapshot;
+		const TSharedRef<TJsonWriter<>> Writer = TJsonWriterFactory<>::Create(&Snapshot);
+		FJsonSerializer::Serialize(JsonObject.ToSharedRef(), Writer);
+
+		UStageMonitorEditorSettings* Settings = GetMutableDefault<UStageMonitorEditorSettings>();
+		Settings->ColumnVisibilitySettings = Snapshot;
+		Settings->SaveConfig();
+	}
+
+	TOptional<TPair<FString, bool>> GetColumnData(const TSharedPtr<FJsonValue>& ColumnValue)
+	{
+		TSharedPtr<FJsonObject>* PossibleColumnObject;
+		if (!ColumnValue->TryGetObject(PossibleColumnObject))
+		{
+			return {};
+		}
+		const TSharedPtr<FJsonObject>& ColumnObject = *PossibleColumnObject;
+
+		FString ColumnIdValue;
+		bool bIsVisible;
+		if (!ColumnObject->TryGetStringField(JsonKeys::ColumnId, ColumnIdValue)
+			|| !ColumnObject->TryGetBoolField(JsonKeys::IsVisible, bIsVisible))
+		{
+			return {};
+		}
+
+		return {{ColumnIdValue, bIsVisible}};
+	}
+
+	void RestoreColumnVisibilityState(TSharedRef<SHeaderRow> HeaderRow)
+	{
+		const UStageMonitorEditorSettings* Settings = GetDefault<UStageMonitorEditorSettings>();
+		TSharedPtr<FJsonObject> JsonObject = MakeShared<FJsonObject>();
+		const TSharedRef<TJsonReader<>> Reader = TJsonReaderFactory<>::Create(Settings->ColumnVisibilitySettings);
+		const bool bIsValid = FJsonSerializer::Deserialize(Reader, JsonObject);
+
+		if (!bIsValid)
+		{
+			return;
+		}
+
+		const TArray<TSharedPtr<FJsonValue>> ColumnObjects = JsonObject->GetArrayField(JsonKeys::Values);
+		for (const TSharedPtr<FJsonValue>& ColumnValue : ColumnObjects)
+		{
+			const TOptional<TPair<FString,bool>> Data = GetColumnData(ColumnValue);
+			if (Data)
+			{
+				HeaderRow->SetShowGeneratedColumn(*Data->Key, Data->Value);
+			}
+		}
+	}
+}
+
 /**
  * SDataProviderListView
  */
 void SDataProviderListView::Construct(const FArguments& InArgs, const TWeakPtr<IStageMonitorSession>& InSession)
 {
-	SortedColumnName = DataProviderListView::HeaderIdName_StageName;
+	using namespace DataProviderListView::HeaderDef;
+
+	SortedColumnName = DataProviderListView::HeaderDef::StageName;
 	SortMode = EColumnSortMode::Ascending;
 
 	AttachToMonitorSession(InSession);
 
+	// Construct the header row.
+	HeaderRow = SNew(SHeaderRow)
+		.OnHiddenColumnsListChanged_Lambda([this]()
+			{
+				if (!bUpdatingColumnVisibility)
+				{
+					UE::StageMonitor::Private::CaptureColumnVisibilityState(HeaderRow.ToSharedRef());
+				}
+			}
+		)
+		.CanSelectGeneratedColumn(true); //To show/hide columns
+
 	Super::Construct
 	(
 		Super::FArguments()
 		.ListItemsSource(&ListItemsSource)
 		.SelectionMode(ESelectionMode::None)
 		.OnGenerateRow(this, &SDataProviderListView::OnGenerateRow)
-		.HeaderRow
-		(
-			SNew(SHeaderRow)
-			.CanSelectGeneratedColumn(true) //To show/hide columns
+		.HeaderRow(HeaderRow)
+	);
 
-			+ SHeaderRow::Column(DataProviderListView::HeaderIdName_State)
-			.FixedWidth(55.f)
-			.HAlignHeader(HAlign_Center)
-			.HAlignCell(HAlign_Center)
-			.DefaultLabel(LOCTEXT("HeaderName_State", "State"))
-			.SortMode(this, &SDataProviderListView::GetColumnSortMode, DataProviderListView::HeaderIdName_State)
-			.OnSort(this, &SDataProviderListView::OnColumnSortModeChanged)
+	// Declare the columns for the header row by referring to the defined table.
+	TArray<FName> ColumnNames;
+	Algo::Transform(StaticColumnHeader, ColumnNames,
+			[](const TPair<FName,FColumnDefinition>& Item)
+			{
+				return Item.Get<0>();
+			});
+
+	ColumnNames.Sort(
+		[](const FName& A, const FName& B)
+		{
+			return StaticColumnHeader.Find(A)->OrderId < StaticColumnHeader.Find(B)->OrderId;
+		});
 
-			+ SHeaderRow::Column(DataProviderListView::HeaderIdName_Status)
-			.FixedWidth(55.f)
+	// Initialize the columns in the proper order.
+	for (const FName& ColumnName : ColumnNames)
+	{
+		const FColumnDefinition* Column = StaticColumnHeader.Find(ColumnName);
+		SHeaderRow::FColumn::FArguments Args = SHeaderRow::FColumn::FArguments()
+			.ColumnId(ColumnName)
+			.DefaultLabel(Column->Label)
 			.HAlignHeader(HAlign_Center)
 			.HAlignCell(HAlign_Center)
-			.DefaultLabel(LOCTEXT("HeaderName_Status", "Status"))
-			.SortMode(this, &SDataProviderListView::GetColumnSortMode, DataProviderListView::HeaderIdName_Status)
-			.OnSort(this, &SDataProviderListView::OnColumnSortModeChanged)
-
-			+ SHeaderRow::Column(DataProviderListView::HeaderIdName_Timecode)
-			.FillWidth(.2f)
-			.DefaultLabel(LOCTEXT("HeaderName_Timecode", "Timecode"))
-
-			+ SHeaderRow::Column(DataProviderListView::HeaderIdName_MachineName)
-			.FillWidth(.2f)
-			.DefaultLabel(LOCTEXT("HeaderName_MachineName", "Machine"))
-			.SortMode(this, &SDataProviderListView::GetColumnSortMode, DataProviderListView::HeaderIdName_MachineName)
-			.OnSort(this, &SDataProviderListView::OnColumnSortModeChanged)
-
-			+ SHeaderRow::Column(DataProviderListView::HeaderIdName_ProcessId)
-			.FillWidth(.15f)
-			.DefaultLabel(LOCTEXT("HeaderName_ProcessId", "Process Id"))
-			.SortMode(this, &SDataProviderListView::GetColumnSortMode, DataProviderListView::HeaderIdName_ProcessId)
-			.OnSort(this, &SDataProviderListView::OnColumnSortModeChanged)
-
-			+ SHeaderRow::Column(DataProviderListView::HeaderIdName_StageName)
-			.FillWidth(.2f)
-			.DefaultLabel(LOCTEXT("HeaderName_StageName", "Stage Name"))
-			.ShouldGenerateWidget(true) //Can't hide this column
-			.SortMode(this, &SDataProviderListView::GetColumnSortMode, DataProviderListView::HeaderIdName_StageName)
-			.OnSort(this, &SDataProviderListView::OnColumnSortModeChanged)
-
-			+ SHeaderRow::Column(DataProviderListView::HeaderIdName_Roles)
-			.FillWidth(.2f)
-			.DefaultLabel(LOCTEXT("HeaderName_Roles", "Roles"))
-			.SortMode(this, &SDataProviderListView::GetColumnSortMode, DataProviderListView::HeaderIdName_Roles)
-			.OnSort(this, &SDataProviderListView::OnColumnSortModeChanged)
-
-			+ SHeaderRow::Column(DataProviderListView::HeaderIdName_AverageFPS)
-			.FillWidth(.2f)
-			.DefaultLabel(LOCTEXT("HeaderName_AverageFPS", "Average FPS"))
-
-			+ SHeaderRow::Column(DataProviderListView::HeaderIdName_IdleTime)
-			.FillWidth(.2f)
-			.DefaultLabel(LOCTEXT("HeaderName_IdleTime", "Idle Time (ms)"))
-
-			+ SHeaderRow::Column(DataProviderListView::HeaderIdName_GameThreadTiming)
-			.FillWidth(.2f)
-			.DefaultLabel(LOCTEXT("HeaderName_GameThread", "Game Thread (ms)"))
-
-			+ SHeaderRow::Column(DataProviderListView::HeaderIdName_RenderThreadTiming)
-			.FillWidth(.2f)
-			.DefaultLabel(LOCTEXT("HeaderName_RenderThread", "Render Thread (ms)"))
-
-			+ SHeaderRow::Column(DataProviderListView::HeaderIdName_GPUTiming)
-			.FillWidth(.2f)
-			.DefaultLabel(LOCTEXT("HeaderName_GPU", "GPU (ms)"))
-
-			+ SHeaderRow::Column(DataProviderListView::HeaderIdName_CPUMem)
-			.FillWidth(.2f)
-			.DefaultLabel(LOCTEXT("HeaderName_CPUMem", "CPU Memory"))
-
-			+ SHeaderRow::Column(DataProviderListView::HeaderIdName_GPUMem)
-			.FillWidth(.2f)
-			.DefaultLabel(LOCTEXT("HeaderName_GPUMem", "GPU Memory"))
-
-			+ SHeaderRow::Column(DataProviderListView::HeaderIdName_AssetsToCompile)
-			.FillWidth(.2f)
-			.DefaultLabel(LOCTEXT("HeaderName_Assets", "Assets To Compile"))
-		)
-	);
+			;
+		if (Column->bSortable)
+		{
+			Args.SortMode(this, &SDataProviderListView::GetColumnSortMode, ColumnName)
+				.OnSort(this, &SDataProviderListView::OnColumnSortModeChanged);
+		}
+		if (Column->ColumnArgumentsFuncPtr)
+		{
+			Column->ColumnArgumentsFuncPtr(Args);
+		}
+		HeaderRow->AddColumn(Args);
+	}
+
+	{
+		TGuardValue<bool> DoNotSave(bUpdatingColumnVisibility, true);
+		SetDefaultColumnVisibilities();
+		UE::StageMonitor::Private::RestoreColumnVisibilityState(HeaderRow.ToSharedRef());
+	}
 
 	RebuildDataProviderList();
 }
 
+void SDataProviderListView::SetDefaultColumnVisibilities()
+{
+	using namespace DataProviderListView::HeaderDef;
+	for (const TPair<FName,FColumnDefinition>& Item : StaticColumnHeader)
+	{
+		HeaderRow->SetShowGeneratedColumn(Item.Key, Item.Value.bIsVisible);
+	}
+}
+
 void SDataProviderListView::RefreshMonitorSession(TWeakPtr<IStageMonitorSession> NewSession)
 {
 	AttachToMonitorSession(NewSession);
@@ -559,7 +804,9 @@ void SDataProviderListView::Tick(const FGeometry& AllottedGeometry, const double
 TSharedRef<ITableRow> SDataProviderListView::OnGenerateRow(FDataProviderTableRowDataPtr InItem, const TSharedRef<STableViewBase>& OwnerTable)
 {
 	TSharedRef<SDataProviderTableRow> Row = SNew(SDataProviderTableRow, OwnerTable)
+		.Padding(2.0f)
 		.Item(InItem);
+
 	ListRowWidgets.Add(Row);
 	return Row;
 }
@@ -633,30 +880,30 @@ void SDataProviderListView::SortProviderList()
 {
 	auto Compare = [this](const FDataProviderTableRowDataPtr& Lhs, const FDataProviderTableRowDataPtr& Rhs, const FName& ColumnName, EColumnSortMode::Type CurrentSortMode)
 	{
-		if (ColumnName == DataProviderListView::HeaderIdName_State)
+		if (ColumnName == DataProviderListView::HeaderDef::State)
 		{
 			return CurrentSortMode == EColumnSortMode::Ascending ? Lhs->CachedState < Rhs->CachedState : Lhs->CachedState > Rhs->CachedState;
 		}
-		if (ColumnName == DataProviderListView::HeaderIdName_Status)
+		if (ColumnName == DataProviderListView::HeaderDef::Status)
 		{
 			const int32 CompareResult = Lhs->CachedPerformanceData.Status < Rhs->CachedPerformanceData.Status;
 			return CurrentSortMode == EColumnSortMode::Ascending ? CompareResult < 0 : CompareResult > 0;
 		}
-		else if (ColumnName == DataProviderListView::HeaderIdName_MachineName)
+		else if (ColumnName == DataProviderListView::HeaderDef::MachineName)
 		{
 			const int32 CompareResult = Lhs->Descriptor.MachineName.Compare(Rhs->Descriptor.MachineName);
 			return CurrentSortMode == EColumnSortMode::Ascending ? CompareResult < 0 : CompareResult > 0;
 		}
-		else if (ColumnName == DataProviderListView::HeaderIdName_StageName)
+		else if (ColumnName == DataProviderListView::HeaderDef::StageName)
 		{
 			const int32 CompareResult = Lhs->Descriptor.FriendlyName.Compare(Rhs->Descriptor.FriendlyName);
 			return CurrentSortMode == EColumnSortMode::Ascending ? CompareResult < 0 : CompareResult > 0;
 		}
-		else if (ColumnName == DataProviderListView::HeaderIdName_ProcessId)
+		else if (ColumnName == DataProviderListView::HeaderDef::ProcessId)
 		{
 			return CurrentSortMode == EColumnSortMode::Ascending ? Lhs->Descriptor.ProcessId < Rhs->Descriptor.ProcessId : Lhs->Descriptor.ProcessId > Rhs->Descriptor.ProcessId;
 		}
-		else if (ColumnName == DataProviderListView::HeaderIdName_Roles)
+		else if (ColumnName == DataProviderListView::HeaderDef::Roles)
 		{
 			if (CachedRoleStringToArray.Contains(Lhs->Descriptor.RolesStringified) == false)
 			{
diff --git a/Engine/Plugins/VirtualProduction/StageMonitoring/Source/StageMonitorEditor/Private/Widgets/SDataProviderListView.h b/Engine/Plugins/VirtualProduction/StageMonitoring/Source/StageMonitorEditor/Private/Widgets/SDataProviderListView.h
--- a/Engine/Plugins/VirtualProduction/StageMonitoring/Source/StageMonitorEditor/Private/Widgets/SDataProviderListView.h
+++ b/Engine/Plugins/VirtualProduction/StageMonitoring/Source/StageMonitorEditor/Private/Widgets/SDataProviderListView.h
@@ -11,7 +11,6 @@ class IStageMonitorSession;
 
 using FDataProviderTableRowDataPtr = TSharedPtr<FDataProviderTableRowData>;
 
-
 /**
  *
  */
@@ -22,14 +21,11 @@ class SDataProviderTableRow : public SMultiColumnTableRow<FDataProviderTableRowD
 public:
 	SLATE_BEGIN_ARGS(SDataProviderTableRow) { }
 		SLATE_ARGUMENT(FDataProviderTableRowDataPtr, Item)
+		SLATE_ARGUMENT(float, Padding)
 	SLATE_END_ARGS()
 
 	void Construct(const FArguments& Args, const TSharedRef<STableViewBase>& OwerTableView);
 
-private:
-
-	/** Handles creation of each columns widget */
-	virtual TSharedRef<SWidget> GenerateWidgetForColumn(const FName& ColumnName) override;
 
 	/** Getters to populate the UI */
 	FText GetStateGlyphs() const;
@@ -51,6 +47,8 @@ private:
 	FText GetAssetsLeftToCompile() const;
 
 private:
+	/** Handles creation of each columns widget */
+	virtual TSharedRef<SWidget> GenerateWidgetForColumn(const FName& ColumnName) override;
 
 	/** Item to display */
 	FDataProviderTableRowDataPtr Item;
@@ -107,6 +105,9 @@ private:
 
 private:
 
+	/** Set the default column visibility state. */
+	void SetDefaultColumnVisibilities();
+
 	/** Widget and data containing info about what's shown in the listview */
 	TArray<FDataProviderTableRowDataPtr> ListItemsSource;
 	TArray<TWeakPtr<SDataProviderTableRow>> ListRowWidgets;
@@ -128,5 +129,8 @@ private:
 
 	/** Cache of stringified roles to array of roles to avoid parsing into array constantly */
 	TMap<FString, TArray<FString>> CachedRoleStringToArray;
+
+	/** Boolean guard to prevent reentrant updates to the column visibility state. */
+	bool bUpdatingColumnVisibility = false;
 };
 
diff --git a/Engine/Plugins/VirtualProduction/StageMonitoring/Source/StageMonitorEditor/StageMonitorEditor.Build.cs b/Engine/Plugins/VirtualProduction/StageMonitoring/Source/StageMonitorEditor/StageMonitorEditor.Build.cs
--- a/Engine/Plugins/VirtualProduction/StageMonitoring/Source/StageMonitorEditor/StageMonitorEditor.Build.cs
+++ b/Engine/Plugins/VirtualProduction/StageMonitoring/Source/StageMonitorEditor/StageMonitorEditor.Build.cs
@@ -24,6 +24,7 @@ public class StageMonitorEditor : ModuleRules
 				"Engine",
 				"GameplayTags",
  				"InputCore",
+				"Json",
 				"LevelEditor",
 				"Projects",
 				"Slate",
