Add a tooltip for rows that contain all of the frame performance data.

diff --git a/Engine/Plugins/VirtualProduction/StageMonitoring/Source/StageMonitorEditor/Private/Widgets/SDataProviderListView.cpp b/Engine/Plugins/VirtualProduction/StageMonitoring/Source/StageMonitorEditor/Private/Widgets/SDataProviderListView.cpp
--- a/Engine/Plugins/VirtualProduction/StageMonitoring/Source/StageMonitorEditor/Private/Widgets/SDataProviderListView.cpp
+++ b/Engine/Plugins/VirtualProduction/StageMonitoring/Source/StageMonitorEditor/Private/Widgets/SDataProviderListView.cpp
@@ -6,6 +6,7 @@
 #include "Dom/JsonObject.h"
 #include "EditorFontGlyphs.h"
 #include "IStageMonitorSession.h"
+#include "Internationalization/Text.h"
 #include "Layout/Visibility.h"
 #include "Misc/App.h"
 #include "Misc/CoreMiscDefines.h"
@@ -30,6 +31,7 @@
  * - A flag to indicate if can be sorted
  * - A flag indicate if it is default visible
  * - A function callback when creating a new row
+ * - A pointer to method on SDataProviderRow table to access FText result for data.
  * - A function callback that can be used to add additional slate args.
  * - An order id that specifies the order in the column is created.
  */
@@ -37,14 +39,17 @@ struct FColumnDefinition
 {
 	using FGenerateRowFuncType = TSharedRef<SWidget>(SDataProviderTableRow*);
 	using FSetColumnArgumentsFuncType = void(SHeaderRow::FColumn::FArguments&);
+	using FWidgetTextGetterType = FText(SDataProviderTableRow::*)() const;
 
 	FColumnDefinition(TAttribute<FText> InLabel, bool bInSortable, bool bInIsVisible,
 					  FGenerateRowFuncType* InGenerateRowFuncPtr,
+					  FWidgetTextGetterType InTextGetterPtr,
 					  FSetColumnArgumentsFuncType* InArgumentsColumnPtr = nullptr)
 		: Label(MoveTemp(InLabel))
 		, bSortable(bInSortable)
 		, bIsVisible(bInIsVisible)
 		, GenerateRowFuncPtr(InGenerateRowFuncPtr)
+		, WidgetTextGetterPtr(InTextGetterPtr)
 		, ColumnArgumentsFuncPtr(InArgumentsColumnPtr)
 	{
 		if (bTrackOrderIdCount)
@@ -68,6 +73,9 @@ struct FColumnDefinition
 	/** Function callback for generating a new row. */
 	FGenerateRowFuncType* GenerateRowFuncPtr = nullptr;
 
+	/** Pointer to member function that gets the FText for data */
+	FWidgetTextGetterType WidgetTextGetterPtr = nullptr;
+
 	/** A function callback for column arguments. */
 	FSetColumnArgumentsFuncType* ColumnArgumentsFuncPtr = nullptr;
 
@@ -138,6 +146,7 @@ namespace DataProviderListView
 							 .ColorAndOpacity(Provider, &SDataProviderTableRow::GetStateColorAndOpacity)
 					 ];
 				 },
+				 nullptr,
 				 [](SHeaderRow::FColumn::FArguments& Args)
 				 {
 					 Args.FixedWidth(55.f);
@@ -160,6 +169,7 @@ namespace DataProviderListView
 						 .ToolTipText(Provider, &SDataProviderTableRow::GetStatusToolTip)
 					 ];
 				 },
+				 &SDataProviderTableRow::GetStatus,
 				 [](SHeaderRow::FColumn::FArguments& Args)
 				 {
 					 Args.FixedWidth(55.f);
@@ -179,7 +189,8 @@ namespace DataProviderListView
 						 SNew(STextBlock)
 						 .Text(Provider, &SDataProviderTableRow::GetTimecode)
 					 ];
-				 }
+				 },
+				 &SDataProviderTableRow::GetTimecode,
 			 }
 			},
 			{MachineName,
@@ -196,7 +207,8 @@ namespace DataProviderListView
 						 SNew(STextBlock)
 						 .Text(Provider, &SDataProviderTableRow::GetMachineName)
 					 ];
-				 }
+				 },
+				 &SDataProviderTableRow::GetMachineName,
 			 }
 			},
 			{ProcessId,
@@ -213,7 +225,8 @@ namespace DataProviderListView
 						 SNew(STextBlock)
 						 .Text(Provider, &SDataProviderTableRow::GetProcessId)
 					 ];
-				 }
+				 },
+				 &SDataProviderTableRow::GetProcessId,
 			 }},
 			{StageName,
 			 {
@@ -229,7 +242,8 @@ namespace DataProviderListView
 						 SNew(STextBlock)
 						 .Text(Provider, &SDataProviderTableRow::GetStageName)
 					 ];
-				 }
+				 },
+				 &SDataProviderTableRow::GetStageName,
 			 }},
 			{Roles,
 			 {
@@ -245,7 +259,8 @@ namespace DataProviderListView
 						 SNew(STextBlock)
 						 .Text(Provider, &SDataProviderTableRow::GetRoles)
 					 ];
-				 }
+				 },
+				 &SDataProviderTableRow::GetRoles,
 			 }},
 			{AverageFPS,
 			 {
@@ -261,7 +276,8 @@ namespace DataProviderListView
 						 SNew(STextBlock)
 						 .Text(MakeAttributeSP(Provider, &SDataProviderTableRow::GetAverageFPS))
 					 ];
-				 }
+				 },
+				 &SDataProviderTableRow::GetAverageFPS,
 			 }},
 			{IdleTime,
 			 {
@@ -277,7 +293,8 @@ namespace DataProviderListView
 						 SNew(STextBlock)
 						 .Text(MakeAttributeSP(Provider, &SDataProviderTableRow::GetIdleTime))
 					 ];
-				 }
+				 },
+				 &SDataProviderTableRow::GetIdleTime,
 			 }},
 			{GameThreadTiming,
 			 {
@@ -293,7 +310,8 @@ namespace DataProviderListView
 						 SNew(STextBlock)
 						 .Text(MakeAttributeSP(Provider, &SDataProviderTableRow::GetGameThreadTiming))
 					 ];
-				 }
+				 },
+				 &SDataProviderTableRow::GetGameThreadTiming,
 			 }
 			},
 			{RenderThreadTiming,
@@ -310,7 +328,8 @@ namespace DataProviderListView
 						 SNew(STextBlock)
 						 .Text(MakeAttributeSP(Provider, &SDataProviderTableRow::GetRenderThreadTiming))
 					 ];
-				 }
+				 },
+				 &SDataProviderTableRow::GetRenderThreadTiming,
 			 }},
 			{GPUTiming,
 			 {
@@ -326,7 +345,8 @@ namespace DataProviderListView
 						 SNew(STextBlock)
 						 .Text(MakeAttributeSP(Provider, &SDataProviderTableRow::GetGPUTiming))
 					 ];
-				 }
+				 },
+				 &SDataProviderTableRow::GetGPUTiming,
 			 }},
 			{CPUMem,
 			 {
@@ -341,8 +361,9 @@ namespace DataProviderListView
 					 [
 						 SNew(STextBlock)
 						 .Text(MakeAttributeSP(Provider, &SDataProviderTableRow::GetCPUMem))
-						 ];
-				 }
+					 ];
+				 },
+				 &SDataProviderTableRow::GetCPUMem,
 			 }},
 			{GPUMem,
 			 {
@@ -358,7 +379,8 @@ namespace DataProviderListView
 						 SNew(STextBlock)
 						 .Text(MakeAttributeSP(Provider, &SDataProviderTableRow::GetGPUMem))
 					 ];
-				 }
+				 },
+				 &SDataProviderTableRow::GetGPUMem,
 			 }
 			},
 			{AssetsToCompile,
@@ -375,14 +397,33 @@ namespace DataProviderListView
 						 SNew(STextBlock)
 						 .Text(MakeAttributeSP(Provider, &SDataProviderTableRow::GetAssetsLeftToCompile))
 					 ];
-				 }
+				 },
+				 &SDataProviderTableRow::GetAssetsLeftToCompile
 			 }
 			}};
 		FColumnDefinition::EndOrderIdCount();
 		return Header;
 	}();
 
-	const FName FirstColumnId(State);
+	TArray<FName> GetSortedColumnNames()
+	{
+		// Declare the columns for the header row by referring to the defined table.
+		TArray<FName> ColumnNames;
+		Algo::Transform(StaticColumnHeader, ColumnNames,
+			[](const TPair<FName,FColumnDefinition>& Item)
+			{
+				return Item.Get<0>();
+			});
+
+		ColumnNames.Sort(
+			[](const FName& A, const FName& B)
+			{
+				return StaticColumnHeader.Find(A)->OrderId < StaticColumnHeader.Find(B)->OrderId;
+			});
+
+		return ColumnNames;
+	}
+
 
 	} // end namespace HeaderDef
 }
@@ -452,13 +493,37 @@ TSharedRef<SWidget> SDataProviderTableRow::GenerateWidgetForColumn(const FName&
 {
 	using namespace DataProviderListView::HeaderDef;
 	const FColumnDefinition* ColumnDef = StaticColumnHeader.Find(ColumnName);
-	if (ColumnDef)
+	if (ColumnDef && ColumnDef->GenerateRowFuncPtr)
 	{
-		return ColumnDef->GenerateRowFuncPtr(this);
+		TSharedRef<SWidget> Widget = ColumnDef->GenerateRowFuncPtr(this);
+		Widget->SetToolTipText(
+			MakeAttributeSP(this, &SDataProviderTableRow::GetToolTipText));
+		return Widget;
 	}
 	return SNullWidget::NullWidget;
 }
 
+FText SDataProviderTableRow::GetToolTipText() const
+{
+	using MemberFunc = FText(SDataProviderTableRow::*)() const;
+	using namespace DataProviderListView::HeaderDef;
+	const SDataProviderTableRow& ThisClass = *this;
+	FTextBuilder RunningTextObject;
+
+	TArray<FName> Columns = GetSortedColumnNames();
+	for (const FName& Column : Columns)
+	{
+		const FColumnDefinition* ColumnDef = StaticColumnHeader.Find(Column);
+		if (ColumnDef && ColumnDef->WidgetTextGetterPtr)
+		{
+			FColumnDefinition::FWidgetTextGetterType GetterPtr = ColumnDef->WidgetTextGetterPtr;
+			RunningTextObject.AppendLineFormat(LOCTEXT("StageMonitorProviderToolTip", "{0} : {1}"),
+											   FText::FromString(Column.ToString()), (ThisClass.*GetterPtr)());
+		}
+	}
+	return RunningTextObject.ToText();
+}
+
 FText SDataProviderTableRow::GetStateGlyphs() const
 {
 	if (Item->CachedPerformanceData.Status !=  EStageMonitorNodeStatus::Ready)
@@ -708,19 +773,7 @@ void SDataProviderListView::Construct(const FArguments& InArgs, const TWeakPtr<I
 		.HeaderRow(HeaderRow)
 	);
 
-	// Declare the columns for the header row by referring to the defined table.
-	TArray<FName> ColumnNames;
-	Algo::Transform(StaticColumnHeader, ColumnNames,
-			[](const TPair<FName,FColumnDefinition>& Item)
-			{
-				return Item.Get<0>();
-			});
-
-	ColumnNames.Sort(
-		[](const FName& A, const FName& B)
-		{
-			return StaticColumnHeader.Find(A)->OrderId < StaticColumnHeader.Find(B)->OrderId;
-		});
+	TArray<FName> ColumnNames = GetSortedColumnNames();
 
 	// Initialize the columns in the proper order.
 	for (const FName& ColumnName : ColumnNames)
diff --git a/Engine/Plugins/VirtualProduction/StageMonitoring/Source/StageMonitorEditor/Private/Widgets/SDataProviderListView.h b/Engine/Plugins/VirtualProduction/StageMonitoring/Source/StageMonitorEditor/Private/Widgets/SDataProviderListView.h
--- a/Engine/Plugins/VirtualProduction/StageMonitoring/Source/StageMonitorEditor/Private/Widgets/SDataProviderListView.h
+++ b/Engine/Plugins/VirtualProduction/StageMonitoring/Source/StageMonitorEditor/Private/Widgets/SDataProviderListView.h
@@ -45,6 +45,7 @@ public:
 	FText GetGPUMem() const;
 	FText GetCPUMem() const;
 	FText GetAssetsLeftToCompile() const;
+	FText GetToolTipText() const;
 
 private:
 	/** Handles creation of each columns widget */
