Show margin time on game and render threads.

diff --git a/Engine/Plugins/VirtualProduction/StageMonitoring/Source/StageDataProvider/Private/FramePerformanceProvider.cpp b/Engine/Plugins/VirtualProduction/StageMonitoring/Source/StageDataProvider/Private/FramePerformanceProvider.cpp
--- a/Engine/Plugins/VirtualProduction/StageMonitoring/Source/StageDataProvider/Private/FramePerformanceProvider.cpp
+++ b/Engine/Plugins/VirtualProduction/StageMonitoring/Source/StageDataProvider/Private/FramePerformanceProvider.cpp
@@ -44,20 +44,24 @@ uint64 TotalGPUResourceSize(const TArray<FStatMessage>& StatMessages)
 FFramePerformanceProviderMessage GetLatestPerformanceData(const	TArray<FStatMessage>& StatMessages)
 {
 	const float GameThreadTime = FPlatformTime::ToMilliseconds(GGameThreadTime);
+	const float GameThreadWaitTime = FPlatformTime::ToMilliseconds(GGameThreadWaitTime);
 	const float RenderThreadTime = FPlatformTime::ToMilliseconds(GRenderThreadTime);
+	const float RenderThreadWaitTime = FPlatformTime::ToMilliseconds(GRenderThreadWaitTime);
 	const float GPUTime = FPlatformTime::ToMilliseconds(GGPUFrameTime);
 	const float IdleTimeMilli = (FApp::GetIdleTime() * 1000.0);
 	FPlatformMemoryStats Stats = FPlatformMemory::GetStats();
 	return FFramePerformanceProviderMessage(
 		EStageMonitorNodeStatus::Unknown,
 		GameThreadTime,
+		GameThreadWaitTime,
 		RenderThreadTime,
+		RenderThreadWaitTime,
 		GPUTime,
 		IdleTimeMilli,
 		Stats.UsedPhysical,
 		TotalGPUResourceSize(StatMessages),
 		CompilationTasksRemaining()
-		);
+	);
 }
 
 }
diff --git a/Engine/Plugins/VirtualProduction/StageMonitoring/Source/StageMonitorCommon/Public/StageMonitorUtils.h b/Engine/Plugins/VirtualProduction/StageMonitoring/Source/StageMonitorCommon/Public/StageMonitorUtils.h
--- a/Engine/Plugins/VirtualProduction/StageMonitoring/Source/StageMonitorCommon/Public/StageMonitorUtils.h
+++ b/Engine/Plugins/VirtualProduction/StageMonitoring/Source/StageMonitorCommon/Public/StageMonitorUtils.h
@@ -32,10 +32,15 @@ public:
 
 	FFramePerformanceProviderMessage() = default;
 
-	FFramePerformanceProviderMessage(EStageMonitorNodeStatus InStatus, float GameThreadTime, float RenderThreadTime, float GPUTime, float IdleTime, uint64 CPUMem, uint64 GPUMem, int32 InAssetsToCompile)
+	FFramePerformanceProviderMessage(EStageMonitorNodeStatus InStatus,
+									 float GameThreadTime, float GameThreadWaitTime,
+									 float RenderThreadTime, float RenderThreadWaitTime,
+									 float GPUTime, float IdleTime, uint64 CPUMem, uint64 GPUMem, int32 InAssetsToCompile)
 		: Status(InStatus)
 		, GameThreadMS(GameThreadTime)
+		, GameThreadWaitMS(GameThreadWaitTime)
 		, RenderThreadMS(RenderThreadTime)
+		, RenderThreadWaitMS(RenderThreadWaitTime)
 		, GPU_MS(GPUTime)
 		, IdleTimeMS(IdleTime)
 		, CPU_MEM(CPUMem)
@@ -62,10 +67,18 @@ public:
 	UPROPERTY(VisibleAnywhere, Category = "Performance", meta = (Unit = "ms"))
 	float GameThreadMS = 0.f;
 
+	/** Current GameThread wait time read from GGameThreadWaitTime in milliseconds */
+	UPROPERTY(VisibleAnywhere, Category = "Performance", meta = (Unit = "ms"))
+	float GameThreadWaitMS = 0.f;
+
 	/** Current RenderThread time read from GRenderThreadTime in milliseconds */
 	UPROPERTY(VisibleAnywhere, Category = "Performance", meta = (Unit = "ms"))
 	float RenderThreadMS = 0.f;
 
+	/** Current RenderThread wait time read from GRenderThreadWaitTime in milliseconds */
+	UPROPERTY(VisibleAnywhere, Category = "Performance", meta = (Unit = "ms"))
+	float RenderThreadWaitMS = 0.f;
+
 	/** Current GPU time read from GGPUFrameTime in milliseconds */
 	UPROPERTY(VisibleAnywhere, Category = "Performance", meta = (Unit = "ms"))
 	float GPU_MS = 0.f;
diff --git a/Engine/Plugins/VirtualProduction/StageMonitoring/Source/StageMonitorEditor/Private/Widgets/SDataProviderListView.cpp b/Engine/Plugins/VirtualProduction/StageMonitoring/Source/StageMonitorEditor/Private/Widgets/SDataProviderListView.cpp
--- a/Engine/Plugins/VirtualProduction/StageMonitoring/Source/StageMonitorEditor/Private/Widgets/SDataProviderListView.cpp
+++ b/Engine/Plugins/VirtualProduction/StageMonitoring/Source/StageMonitorEditor/Private/Widgets/SDataProviderListView.cpp
@@ -117,9 +117,12 @@ namespace DataProviderListView
 	const FName StageName = TEXT("StageName");
 	const FName Roles = TEXT("Roles");
 	const FName AverageFPS = TEXT("AverageFPS");
+	const FName EstimatedMaxFPS = TEXT("EstimatedMaxFPS");
 	const FName IdleTime = TEXT("IdleTime");
 	const FName GameThreadTiming = TEXT("GameThreadTiming");
+	const FName GameThreadWaitTiming = TEXT("GameThreadWaitTiming");
 	const FName RenderThreadTiming = TEXT("RenderThreadTiming");
+	const FName RenderThreadWaitTiming = TEXT("RenderThreadWaitTiming");
 	const FName GPUTiming = TEXT("GPUTiming");
 	const FName CPUMem = TEXT("CPUMem");
 	const FName GPUMem = TEXT("GPUMem");
@@ -279,6 +282,23 @@ namespace DataProviderListView
 				 },
 				 &SDataProviderTableRow::GetAverageFPS,
 			 }},
+			{EstimatedMaxFPS,
+			 {
+				 LOCTEXT("HeaderName_AverageFPS", "Estimated Max FPS"), IsNotSortable, IsVisible,
+				 [](SDataProviderTableRow* Provider) -> TSharedRef<SWidget>
+				 {
+					 return SNew(SHorizontalBox)
+					 + SHorizontalBox::Slot()
+					 .HAlign(HAlign_Left)
+					 .VAlign(VAlign_Center)
+					 .AutoWidth()
+					 [
+						 SNew(STextBlock)
+						 .Text(MakeAttributeSP(Provider, &SDataProviderTableRow::GetEstimatedMaxFPS))
+					 ];
+				 },
+				 &SDataProviderTableRow::GetEstimatedMaxFPS,
+			 }},
 			{IdleTime,
 			 {
 				 LOCTEXT("HeaderName_IdleTime", "Idle Time (ms)"), IsNotSortable, IsNotVisible,
@@ -314,6 +334,24 @@ namespace DataProviderListView
 				 &SDataProviderTableRow::GetGameThreadTiming,
 			 }
 			},
+			{GameThreadWaitTiming,
+			 {
+				 LOCTEXT("HeaderName_GameThread", "Game Wait Time (ms)"), IsNotSortable, IsNotVisible,
+				 [](SDataProviderTableRow* Provider) -> TSharedRef<SWidget>
+				 {
+					 return SNew(SHorizontalBox)
+					 + SHorizontalBox::Slot()
+					 .HAlign(HAlign_Left)
+					 .VAlign(VAlign_Center)
+					 .AutoWidth()
+					 [
+						 SNew(STextBlock)
+						 .Text(MakeAttributeSP(Provider, &SDataProviderTableRow::GetGameThreadWaitTiming))
+					 ];
+				 },
+				 &SDataProviderTableRow::GetGameThreadWaitTiming,
+			 }
+			},
 			{RenderThreadTiming,
 			 {
 				 LOCTEXT("HeaderName_RenderThread", "Render Thread (ms)"), IsNotSortable, IsVisible,
@@ -331,6 +369,23 @@ namespace DataProviderListView
 				 },
 				 &SDataProviderTableRow::GetRenderThreadTiming,
 			 }},
+			{RenderThreadWaitTiming,
+			 {
+				 LOCTEXT("HeaderName_RenderWaitThread", "Render Wait Time (ms)"), IsNotSortable, IsNotVisible,
+				 [](SDataProviderTableRow* Provider) -> TSharedRef<SWidget>
+				 {
+					 return SNew(SHorizontalBox)
+					 + SHorizontalBox::Slot()
+					 .HAlign(HAlign_Left)
+					 .VAlign(VAlign_Center)
+					 .AutoWidth()
+					 [
+						 SNew(STextBlock)
+						 .Text(MakeAttributeSP(Provider, &SDataProviderTableRow::GetRenderThreadWaitTiming))
+					 ];
+				 },
+				 &SDataProviderTableRow::GetRenderThreadWaitTiming,
+			 }},
 			{GPUTiming,
 			 {
 				 LOCTEXT("HeaderName_GPU", "GPU (ms)"), IsNotSortable, IsVisible,
@@ -631,6 +686,18 @@ FText SDataProviderTableRow::GetAverageFPS() const
 	return FText::AsNumber(Item->CachedPerformanceData.AverageFPS);
 }
 
+FText SDataProviderTableRow::GetEstimatedMaxFPS() const
+{
+	const FFramePerformanceProviderMessage& Data = Item->CachedPerformanceData;
+	const float AverageMS = Data.AverageFPS > 0 ? 1000.0 / Data.AverageFPS : 0;
+	const float ThreadWaitMS = FMath::Min(Data.GameThreadWaitMS, Data.RenderThreadWaitMS);
+	const float WaitMS = Data.IdleTimeMS > ThreadWaitMS ? Data.IdleTimeMS - ThreadWaitMS : ThreadWaitMS;
+
+	const float EstimatedFrameTimeMS = AverageMS > WaitMS ? (AverageMS - WaitMS) : AverageMS;
+	const float EstimatedMaxFPS = EstimatedFrameTimeMS > 0 ? 1000.f / EstimatedFrameTimeMS : 0;
+	return FText::AsNumber(EstimatedMaxFPS);
+}
+
 FText SDataProviderTableRow::GetIdleTime() const
 {
 	return FText::AsNumber(Item->CachedPerformanceData.IdleTimeMS);
@@ -641,11 +708,21 @@ FText SDataProviderTableRow::GetGameThreadTiming() const
 	return FText::AsNumber(Item->CachedPerformanceData.GameThreadMS);
 }
 
+FText SDataProviderTableRow::GetGameThreadWaitTiming() const
+{
+	return FText::AsNumber(Item->CachedPerformanceData.GameThreadWaitMS);
+}
+
 FText SDataProviderTableRow::GetRenderThreadTiming() const
 {
 	return FText::AsNumber(Item->CachedPerformanceData.RenderThreadMS);
 }
 
+FText SDataProviderTableRow::GetRenderThreadWaitTiming() const
+{
+	return FText::AsNumber(Item->CachedPerformanceData.RenderThreadWaitMS);
+}
+
 FText SDataProviderTableRow::GetGPUTiming() const
 {
 	return FText::AsNumber(Item->CachedPerformanceData.GPU_MS);
diff --git a/Engine/Plugins/VirtualProduction/StageMonitoring/Source/StageMonitorEditor/Private/Widgets/SDataProviderListView.h b/Engine/Plugins/VirtualProduction/StageMonitoring/Source/StageMonitorEditor/Private/Widgets/SDataProviderListView.h
--- a/Engine/Plugins/VirtualProduction/StageMonitoring/Source/StageMonitorEditor/Private/Widgets/SDataProviderListView.h
+++ b/Engine/Plugins/VirtualProduction/StageMonitoring/Source/StageMonitorEditor/Private/Widgets/SDataProviderListView.h
@@ -38,9 +38,12 @@ public:
 	FText GetStageName() const;
 	FText GetRoles() const;
 	FText GetAverageFPS() const;
+	FText GetEstimatedMaxFPS() const;
 	FText GetIdleTime() const;
 	FText GetGameThreadTiming() const;
+	FText GetGameThreadWaitTiming() const;
 	FText GetRenderThreadTiming() const;
+	FText GetRenderThreadWaitTiming() const;
 	FText GetGPUTiming() const;
 	FText GetGPUMem() const;
 	FText GetCPUMem() const;
diff --git a/Engine/Source/Runtime/Engine/Private/UnrealClient.cpp b/Engine/Source/Runtime/Engine/Private/UnrealClient.cpp
--- a/Engine/Source/Runtime/Engine/Private/UnrealClient.cpp
+++ b/Engine/Source/Runtime/Engine/Private/UnrealClient.cpp
@@ -1832,6 +1832,7 @@ void FViewport::Draw( bool bShouldPresent /*= true */)
 							uint32 ThreadTime	= CurrentTime - Lastimestamp;
 							// add any stalls via sleep or fevent
 							GGameThreadTime		= (ThreadTime > GameThread.Waits) ? (ThreadTime - GameThread.Waits) : ThreadTime;
+							GGameThreadWaitTime = GameThread.Waits;
 						}
 						else
 						{
diff --git a/Engine/Source/Runtime/RenderCore/Private/RenderingThread.cpp b/Engine/Source/Runtime/RenderCore/Private/RenderingThread.cpp
--- a/Engine/Source/Runtime/RenderCore/Private/RenderingThread.cpp
+++ b/Engine/Source/Runtime/RenderCore/Private/RenderingThread.cpp
@@ -285,13 +285,13 @@ void TickRenderingTickables()
 
 /** How many cycles the renderthread used (excluding idle time). It's set once per frame in FViewport::Draw. */
 uint32 GRenderThreadTime = 0;
+/** How many cycles of wait time renderthread used. It's set once per frame in FViewport::Draw. */
+uint32 GRenderThreadWaitTime = 0;
 /** How many cycles the rhithread used (excluding idle time). */
 uint32 GRHIThreadTime = 0;
 /** How many cycles the renderthread used, including dependent wait time. */
 uint32 GRenderThreadTimeCriticalPath = 0;
 
-
-
 /** The RHI thread runnable object. */
 class FRHIThread : public FRunnable
 {
@@ -1149,6 +1149,10 @@ bool FRenderCommandFence::IsFenceComplete() const
 
 /** How many cycles the gamethread used (excluding idle time). It's set once per frame in FViewport::Draw. */
 uint32 GGameThreadTime = 0;
+
+/** How much idle time on the game thread. It's set once per frame in FViewport::Draw. */
+uint32 GGameThreadWaitTime = 0;
+
 /** How many cycles it took to swap buffers to present the frame. */
 uint32 GSwapBufferTime = 0;
 
diff --git a/Engine/Source/Runtime/RenderCore/Public/RenderTimer.h b/Engine/Source/Runtime/RenderCore/Public/RenderTimer.h
--- a/Engine/Source/Runtime/RenderCore/Public/RenderTimer.h
+++ b/Engine/Source/Runtime/RenderCore/Public/RenderTimer.h
@@ -107,12 +107,18 @@ extern RENDERCORE_API FInputLatencyTimer GInputLatencyTimer;
 /** How many cycles the renderthread used (excluding idle time). It's set once per frame in FViewport::Draw. */
 extern RENDERCORE_API uint32 GRenderThreadTime;
 
+/** How many cycles the renderthread was waiting. It's set once per frame in FViewport::Draw. */
+extern RENDERCORE_API uint32 GRenderThreadWaitTime;
+
 /** How many cycles the rhithread used (excluding idle time). */
 extern RENDERCORE_API uint32 GRHIThreadTime;
 
 /** How many cycles the gamethread used (excluding idle time). It's set once per frame in FViewport::Draw. */
 extern RENDERCORE_API uint32 GGameThreadTime;
 
+/** How much idle time in the game thread. It's set once per frame in FViewport::Draw. */
+extern RENDERCORE_API uint32 GGameThreadWaitTime;
+
 /** How many cycles it took to swap buffers to present the frame. */
 extern RENDERCORE_API uint32 GSwapBufferTime;
 
diff --git a/Engine/Source/Runtime/SlateRHIRenderer/Private/SlateRHIRenderer.cpp b/Engine/Source/Runtime/SlateRHIRenderer/Private/SlateRHIRenderer.cpp
--- a/Engine/Source/Runtime/SlateRHIRenderer/Private/SlateRHIRenderer.cpp
+++ b/Engine/Source/Runtime/SlateRHIRenderer/Private/SlateRHIRenderer.cpp
@@ -1548,6 +1548,7 @@ void FSlateRHIRenderer::DrawWindow_RenderThread(FRHICommandListImmediate& RHICmd
 
 	SET_CYCLE_COUNTER(STAT_RenderingIdleTime, RenderThreadIdle);
 	GRenderThreadTime = (ThreadTime > RenderThreadIdle) ? (ThreadTime - RenderThreadIdle) : ThreadTime;
+	GRenderThreadWaitTime = RenderThreadIdle;
 
 	// Compute GRenderThreadTimeCriticalPath
 	uint32 RenderThreadNonCriticalPathIdle = RenderThreadIdle - RenderThread.WaitsCriticalPath;
