Fix localization issue.

diff --git a/Engine/Plugins/VirtualProduction/StageMonitoring/Source/StageMonitorEditor/Private/Widgets/SDataProviderListView.cpp b/Engine/Plugins/VirtualProduction/StageMonitoring/Source/StageMonitorEditor/Private/Widgets/SDataProviderListView.cpp
--- a/Engine/Plugins/VirtualProduction/StageMonitoring/Source/StageMonitorEditor/Private/Widgets/SDataProviderListView.cpp
+++ b/Engine/Plugins/VirtualProduction/StageMonitoring/Source/StageMonitorEditor/Private/Widgets/SDataProviderListView.cpp
@@ -284,7 +284,7 @@ namespace DataProviderListView
 			 }},
 			{EstimatedMaxFPS,
 			 {
-				 LOCTEXT("HeaderName_AverageFPS", "Estimated Max FPS"), IsNotSortable, IsVisible,
+				 LOCTEXT("HeaderName_EstimatedMaxFPS", "Estimated Max FPS"), IsNotSortable, IsVisible,
 				 [](SDataProviderTableRow* Provider) -> TSharedRef<SWidget>
 				 {
 					 return SNew(SHorizontalBox)
