Estimated Max FPS improvements

diff --git a/Engine/Plugins/VirtualProduction/StageMonitoring/Source/StageDataProvider/Private/FramePerformanceProvider.cpp b/Engine/Plugins/VirtualProduction/StageMonitoring/Source/StageDataProvider/Private/FramePerformanceProvider.cpp
--- a/Engine/Plugins/VirtualProduction/StageMonitoring/Source/StageDataProvider/Private/FramePerformanceProvider.cpp
+++ b/Engine/Plugins/VirtualProduction/StageMonitoring/Source/StageDataProvider/Private/FramePerformanceProvider.cpp
@@ -41,13 +41,20 @@ uint64 TotalGPUResourceSize(const TArray<FStatMessage>& StatMessages)
 	return TotalSize;
 }
 
+static float GAverageGPU = 0;
+
+void UpdateAverageGPUUsage()
+{
+	const float GPUTime = FPlatformTime::ToMilliseconds(GGPUFrameTime);
+	GAverageGPU = 0.75*GAverageGPU + 0.25*GPUTime;
+}
+
 FFramePerformanceProviderMessage GetLatestPerformanceData(const	TArray<FStatMessage>& StatMessages)
 {
 	const float GameThreadTime = FPlatformTime::ToMilliseconds(GGameThreadTime);
 	const float GameThreadWaitTime = FPlatformTime::ToMilliseconds(GGameThreadWaitTime);
 	const float RenderThreadTime = FPlatformTime::ToMilliseconds(GRenderThreadTime);
 	const float RenderThreadWaitTime = FPlatformTime::ToMilliseconds(GRenderThreadWaitTime);
-	const float GPUTime = FPlatformTime::ToMilliseconds(GGPUFrameTime);
 	const float IdleTimeMilli = (FApp::GetIdleTime() * 1000.0);
 	FPlatformMemoryStats Stats = FPlatformMemory::GetStats();
 	return FFramePerformanceProviderMessage(
@@ -56,7 +63,7 @@ FFramePerformanceProviderMessage GetLatestPerformanceData(const	TArray<FStatMess
 		GameThreadWaitTime,
 		RenderThreadTime,
 		RenderThreadWaitTime,
-		GPUTime,
+		GAverageGPU,
 		IdleTimeMilli,
 		Stats.UsedPhysical,
 		TotalGPUResourceSize(StatMessages),
@@ -99,6 +106,7 @@ public:
 				GetPermanentStats(Stats);
 			}
 
+			UE::FramePerformanceProvider::Private::UpdateAverageGPUUsage();
 			IStageDataProvider::SendMessage<FFramePerformanceProviderMessage>(EStageMessageFlags::None, GetFramePerformanceData());
 
 			FPlatformProcess::Sleep(UpdateFrequency);
diff --git a/Engine/Plugins/VirtualProduction/StageMonitoring/Source/StageMonitorEditor/Private/Widgets/SDataProviderListView.cpp b/Engine/Plugins/VirtualProduction/StageMonitoring/Source/StageMonitorEditor/Private/Widgets/SDataProviderListView.cpp
--- a/Engine/Plugins/VirtualProduction/StageMonitoring/Source/StageMonitorEditor/Private/Widgets/SDataProviderListView.cpp
+++ b/Engine/Plugins/VirtualProduction/StageMonitoring/Source/StageMonitorEditor/Private/Widgets/SDataProviderListView.cpp
@@ -689,12 +689,11 @@ FText SDataProviderTableRow::GetAverageFPS() const
 FText SDataProviderTableRow::GetEstimatedMaxFPS() const
 {
 	const FFramePerformanceProviderMessage& Data = Item->CachedPerformanceData;
-	const float AverageMS = Data.AverageFPS > 0 ? 1000.0 / Data.AverageFPS : 0;
-	const float ThreadWaitMS = FMath::Min(Data.GameThreadWaitMS, Data.RenderThreadWaitMS);
-	const float WaitMS = Data.IdleTimeMS > ThreadWaitMS ? Data.IdleTimeMS - ThreadWaitMS : ThreadWaitMS;
-
-	const float EstimatedFrameTimeMS = AverageMS > WaitMS ? (AverageMS - WaitMS) : AverageMS;
-	const float EstimatedMaxFPS = EstimatedFrameTimeMS > 0 ? 1000.f / EstimatedFrameTimeMS : 0;
+	const float TotalMS = Data.GameThreadMS + Data.GameThreadWaitMS;
+	const float EstimatedGPUIdle = TotalMS > Data.GPU_MS ? (TotalMS - Data.GPU_MS) : TotalMS;
+	const float EstimatedMinMS = TotalMS > EstimatedGPUIdle ? (TotalMS - EstimatedGPUIdle) : TotalMS;
+	const float ClampedEstimatedMS = FMath::Clamp(EstimatedMinMS, Data.GPU_MS, EstimatedMinMS);
+	const float EstimatedMaxFPS = ClampedEstimatedMS > 0 ? 1000.f / ClampedEstimatedMS : 0;
 	return FText::AsNumber(EstimatedMaxFPS);
 }
 
