Fix stage monitor hitching when receiving data from provider with invalid timecode

diff --git a/Engine/Plugins/VirtualProduction/StageMonitoring/Source/StageMonitor/Private/StageMonitorSession.cpp b/Engine/Plugins/VirtualProduction/StageMonitoring/Source/StageMonitor/Private/StageMonitorSession.cpp
--- a/Engine/Plugins/VirtualProduction/StageMonitoring/Source/StageMonitor/Private/StageMonitorSession.cpp
+++ b/Engine/Plugins/VirtualProduction/StageMonitoring/Source/StageMonitor/Private/StageMonitorSession.cpp
@@ -129,6 +129,8 @@ void FStageMonitorSession::UpdateIdentifierMapping(const FGuid& OldIdentifier, c
 
 void FStageMonitorSession::AddProviderMessage(UScriptStruct* Type, const FStageProviderMessage* MessageData)
 {
+	TRACE_CPUPROFILER_EVENT_SCOPE(FStageMonitorSession::AddProviderMessage);
+
 	if (MessageData == nullptr)
 	{
 		return;
@@ -325,24 +327,39 @@ void FStageMonitorSession::UpdateProviderLatestEntry(const FGuid& Identifier, US
 
 void FStageMonitorSession::InsertNewEntry(TSharedPtr<FStageDataEntry> NewEntry)
 {
+	TRACE_CPUPROFILER_EVENT_SCOPE(FStageMonitorSession::InsertNewEntry);
+
 	const FStageProviderMessage* Message = reinterpret_cast<const FStageProviderMessage*>(NewEntry->Data->GetStructMemory());
 	check(Message);
 
 	//We are always inserting in order so we shouldn't have to navigate in our list much
 	const double NewFrameSeconds = Message->FrameTime.AsSeconds();
-	int32 EntryIndex = Entries.Num() - 1;
-	for (; EntryIndex >= 0; --EntryIndex)
+	const double NewFrameDay = FMath::Floor(Message->DateTime.GetJulianDay());
+
+	if (NewFrameSeconds >= 0)
 	{
-		const FStageProviderMessage* ThisEntry = reinterpret_cast<const FStageProviderMessage*>(Entries[EntryIndex]->Data->GetStructMemory());
-		const double ThisFrameSeconds = ThisEntry->FrameTime.AsSeconds();
-		if (ThisFrameSeconds <= NewFrameSeconds)
+		int32 EntryIndex = Entries.Num() - 1;
+		for (; EntryIndex >= 0; --EntryIndex)
 		{
-			break;
+			const FStageProviderMessage* ThisEntry = reinterpret_cast<const FStageProviderMessage*>(Entries[EntryIndex]->Data->GetStructMemory());
+		
+			const double ThisFrameSeconds = ThisEntry->FrameTime.AsSeconds();
+			const double ThisFrameDay = FMath::Floor(ThisEntry->DateTime.GetJulianDay());
+
+			if (ThisFrameSeconds <= NewFrameSeconds || NewFrameDay > ThisFrameDay)
+			{
+				break;
+			}
 		}
-	}
 
-	const int32 InsertIndex = EntryIndex + 1;
-	Entries.Insert(NewEntry, InsertIndex);
+		const int32 InsertIndex = EntryIndex + 1;
+		Entries.Insert(NewEntry, InsertIndex);
+	}
+	else
+	{
+		// Invalid entries with negative timecode get added to the end of the list to avoid going through all the list when they're inserted.
+		Entries.Add(NewEntry);
+	}
 }
 
 void FStageMonitorSession::GetAllEntries(TArray<TSharedPtr<FStageDataEntry>>& OutEntries)
diff --git a/Engine/Plugins/VirtualProduction/StageMonitoring/Source/StageMonitorEditor/Private/Widgets/SDataProviderActivities.cpp b/Engine/Plugins/VirtualProduction/StageMonitoring/Source/StageMonitorEditor/Private/Widgets/SDataProviderActivities.cpp
--- a/Engine/Plugins/VirtualProduction/StageMonitoring/Source/StageMonitorEditor/Private/Widgets/SDataProviderActivities.cpp
+++ b/Engine/Plugins/VirtualProduction/StageMonitoring/Source/StageMonitorEditor/Private/Widgets/SDataProviderActivities.cpp
@@ -199,15 +199,23 @@ void SDataProviderActivities::InsertActivity(FDataProviderActivityPtr Activity)
 	//We might receive messages out of order, make sure we're displaying them in order of source timecode
 	const FStageProviderMessage* ThisActivity = reinterpret_cast<const FStageProviderMessage*>(Activity->Data->GetStructMemory());
 	const double NewFrameSeconds = ThisActivity->FrameTime.AsSeconds();
+	const double NewFrameDay = FMath::Floor(ThisActivity->DateTime.GetJulianDay());
+
 	int32 EntryIndex = 0;
 
-	for (; EntryIndex < FilteredActivities.Num(); ++EntryIndex)
+	// If an activity has a negative timecode, just add it to the top of the list to avoid comparing against all the entries in the list.
+	if (NewFrameSeconds >= 0)
 	{
-		const FStageProviderMessage* ThisEntry = reinterpret_cast<const FStageProviderMessage*>(FilteredActivities[EntryIndex]->Data->GetStructMemory());
-		const double ThisFrameSeconds = ThisEntry->FrameTime.AsSeconds();
-		if (NewFrameSeconds >= ThisFrameSeconds)
+		for (; EntryIndex < FilteredActivities.Num(); ++EntryIndex)
 		{
-			break;
+			const FStageProviderMessage* ThisEntry = reinterpret_cast<const FStageProviderMessage*>(FilteredActivities[EntryIndex]->Data->GetStructMemory());
+			const double ThisFrameSeconds = ThisEntry->FrameTime.AsSeconds();
+			const double ThisFrameDay = FMath::Floor(ThisEntry->DateTime.GetJulianDay());
+
+			if (NewFrameSeconds >= ThisFrameSeconds || NewFrameDay > ThisFrameDay)
+			{
+				break;
+			}
 		}
 	}
 
@@ -368,7 +376,14 @@ FText SDataProviderActivitiesTableRow::GetTimecode() const
 	{
 		check(Item->Data->GetStruct()->IsChildOf(FStageProviderMessage::StaticStruct()));
 		FStageProviderMessage* Data = reinterpret_cast<FStageProviderMessage*>(Item->Data->GetStructMemory());
-		return FText::FromString(FTimecode::FromFrameNumber(Data->FrameTime.Time.GetFrame(), Data->FrameTime.Rate).ToString());
+		if (Data->FrameTime.AsSeconds() >= 0)
+		{
+			return FText::FromString(FTimecode::FromFrameNumber(Data->FrameTime.Time.GetFrame(), Data->FrameTime.Rate).ToString());
+		}
+		else
+		{
+			return LOCTEXT("InvalidTimecode", "Invalid Timecode");
+		}
 	}
 
 	return FText::GetEmpty();
diff --git a/Engine/Source/Runtime/VirtualProduction/StageDataCore/Private/StageMessages.cpp b/Engine/Source/Runtime/VirtualProduction/StageDataCore/Private/StageMessages.cpp
--- a/Engine/Source/Runtime/VirtualProduction/StageDataCore/Private/StageMessages.cpp
+++ b/Engine/Source/Runtime/VirtualProduction/StageDataCore/Private/StageMessages.cpp
@@ -14,6 +14,8 @@ FStageProviderMessage::FStageProviderMessage()
 {
 	//Common setup of timecode for all provider messages
 	TOptional<FQualifiedFrameTime> CurrentFrameTime = FApp::GetCurrentFrameTime();
+	DateTime = FDateTime::Now();
+	
 	if (CurrentFrameTime.IsSet())
 	{
 		FrameTime = CurrentFrameTime.GetValue();
diff --git a/Engine/Source/Runtime/VirtualProduction/StageDataCore/Public/StageMessages.h b/Engine/Source/Runtime/VirtualProduction/StageDataCore/Public/StageMessages.h
--- a/Engine/Source/Runtime/VirtualProduction/StageDataCore/Public/StageMessages.h
+++ b/Engine/Source/Runtime/VirtualProduction/StageDataCore/Public/StageMessages.h
@@ -185,6 +185,10 @@ public:
 	/** FrameTime of the sender. It's expected to have all stage machines using the same timecode provider to play in the same referential*/
 	UPROPERTY()
 	FQualifiedFrameTime FrameTime;
+
+	/** DateTime of the sender. Used to keep track of messages order through Timecode rollover. */
+	UPROPERTY()
+	FDateTime DateTime;
 };
 
 /**
