Fixed some unusual Forward/MoveTemp usage which was now complaining with the new UE_INTRINSIC_CAST feature.

diff --git a/Engine/Plugins/VirtualProduction/StageMonitoring/Source/StageMonitor/Private/StageMonitor.h b/Engine/Plugins/VirtualProduction/StageMonitoring/Source/StageMonitor/Private/StageMonitor.h
--- a/Engine/Plugins/VirtualProduction/StageMonitoring/Source/StageMonitor/Private/StageMonitor.h
+++ b/Engine/Plugins/VirtualProduction/StageMonitoring/Source/StageMonitor/Private/StageMonitor.h
@@ -73,18 +73,8 @@ private:
 	{
 #if ENABLE_STAGEMONITOR_LOGGING
 		static_assert(TIsDerivedFrom<MessageType, FStageMonitorBaseMessage>::IsDerived, "MessageType must be a FStageMonitorBaseMessage derived UStruct.");
-		if constexpr(sizeof...(Args) == 1 && std::is_same<MessageType,TTupleElement<0,TTuple<Args...>>>::value)
-		{
-			MessageType TempObj = Forward<MessageType>(MoveTempIfPossible(TTuple<Args...>::Get<0>(Tie(args...))));
-			return SendMessageInternal(&TempObj, MessageType::StaticStruct(), Flags);
-		}
-		else
-		{
-			MessageType TempObj = MessageType(Forward<Args>(MoveTempIfPossible(args))...);
-			return SendMessageInternal(&TempObj, MessageType::StaticStruct(), Flags);
-		}
-
-
+		MessageType TempObj(Forward<Args>(args)...);
+		return SendMessageInternal(&TempObj, MessageType::StaticStruct(), Flags);
 #endif
 		return false;
 	}
@@ -95,7 +85,7 @@ private:
 	{
 #if ENABLE_STAGEMONITOR_LOGGING
 		static_assert(TIsDerivedFrom<MessageType, FStageMonitorBaseMessage>::IsDerived, "MessageType must be a FStageMonitorBaseMessage derived UStruct.");
-		MessageType TempObj = MessageType(Forward<Args>(MoveTempIfPossible(args))...);
+		MessageType TempObj(Forward<Args>(args)...);
 		PublishMessageInternal(&TempObj, MessageType::StaticStruct());
 #endif
 	}
diff --git a/Engine/Source/Runtime/VirtualProduction/StageDataCore/Public/IStageDataProvider.h b/Engine/Source/Runtime/VirtualProduction/StageDataCore/Public/IStageDataProvider.h
--- a/Engine/Source/Runtime/VirtualProduction/StageDataCore/Public/IStageDataProvider.h
+++ b/Engine/Source/Runtime/VirtualProduction/StageDataCore/Public/IStageDataProvider.h
@@ -32,8 +32,8 @@ public:
 	static bool SendMessage(EStageMessageFlags Flags, Args&&... args)
 	{
 #if ENABLE_STAGEMONITOR_LOGGING
-		MessageType TempObj = MessageType(Forward<Args>(MoveTempIfPossible(args))...);
-		return SendMessage(Forward<MessageType>(TempObj), Flags);
+		MessageType TempObj(Forward<Args>(args)...);
+		return SendMessage(MoveTempIfPossible(TempObj), Flags);
 #endif
 
 		return false;
