StageMonitor: Enable data reporting in packaged dev builds.

This will allow packaged nDisplay nodes to report back to a Monitor their stats.

Also added a few missing #if STATS in the performance monitor, which failed compilation when generating shipping builds once its module was declared as Runtime instead of UncookedOnly.

#rn Stage Monitor can now send data on packaged development builds.

diff --git a/Engine/Plugins/VirtualProduction/StageMonitoring/Source/StageDataProvider/Private/FramePerformanceProvider.cpp b/Engine/Plugins/VirtualProduction/StageMonitoring/Source/StageDataProvider/Private/FramePerformanceProvider.cpp
--- a/Engine/Plugins/VirtualProduction/StageMonitoring/Source/StageDataProvider/Private/FramePerformanceProvider.cpp
+++ b/Engine/Plugins/VirtualProduction/StageMonitoring/Source/StageDataProvider/Private/FramePerformanceProvider.cpp
@@ -26,6 +26,7 @@ int32 CompilationTasksRemaining()
 	return Manager.GetNumRemainingAssets();
 }
 
+#if STATS
 uint64 TotalGPUResourceSize(const TArray<FStatMessage>& StatMessages)
 {
 	uint64 TotalSize = 0;
@@ -40,6 +41,7 @@ uint64 TotalGPUResourceSize(const TArray<FStatMessage>& StatMessages)
 	}
 	return TotalSize;
 }
+#endif // STATS
 
 static float GAverageGPU = 0;
 
@@ -49,14 +51,27 @@ void UpdateAverageGPUUsage()
 	GAverageGPU = 0.75*GAverageGPU + 0.25*GPUTime;
 }
 
+#if STATS
 FFramePerformanceProviderMessage GetLatestPerformanceData(const	TArray<FStatMessage>& StatMessages)
+#else
+FFramePerformanceProviderMessage GetLatestPerformanceData()
+#endif // STATS
 {
 	const float GameThreadTime = FPlatformTime::ToMilliseconds(GGameThreadTime);
 	const float GameThreadWaitTime = FPlatformTime::ToMilliseconds(GGameThreadWaitTime);
 	const float RenderThreadTime = FPlatformTime::ToMilliseconds(GRenderThreadTime);
 	const float RenderThreadWaitTime = FPlatformTime::ToMilliseconds(GRenderThreadWaitTime);
 	const float IdleTimeMilli = (FApp::GetIdleTime() * 1000.0);
+
+#if STATS
 	FPlatformMemoryStats Stats = FPlatformMemory::GetStats();
+	const uint64 CPUMem = Stats.UsedPhysical;
+	const uint64 GPUMem = TotalGPUResourceSize(StatMessages);
+#else
+	constexpr uint64 CPUMem = 0;
+	constexpr uint64 GPUMem = 0;
+#endif // STATS
+
 	return FFramePerformanceProviderMessage(
 		EStageMonitorNodeStatus::Unknown,
 		GameThreadTime,
@@ -65,8 +80,8 @@ FFramePerformanceProviderMessage GetLatestPerformanceData(const	TArray<FStatMess
 		RenderThreadWaitTime,
 		GAverageGPU,
 		IdleTimeMilli,
-		Stats.UsedPhysical,
-		TotalGPUResourceSize(StatMessages),
+		CPUMem,
+		GPUMem,
 		CompilationTasksRemaining()
 	);
 }
@@ -92,11 +107,14 @@ public:
 	// FRunnable functions
 	virtual uint32 Run() override
 	{
+#if STATS
 		double LastRHIUpdate = FPlatformTime::Seconds();
+#endif // STATS
 
 		bStopped = false;
 		do
 		{
+#if STATS
 			const double CurrentPlatformTimeInSeconds = FPlatformTime::Seconds();
 			if (CurrentPlatformTimeInSeconds > (UpdateRHIResourcesFrequency + LastRHIUpdate)
 				|| Stats.IsEmpty())
@@ -105,6 +123,7 @@ public:
 				LastRHIUpdate = CurrentPlatformTimeInSeconds;
 				GetPermanentStats(Stats);
 			}
+#endif // STATS
 
 			UE::FramePerformanceProvider::Private::UpdateAverageGPUUsage();
 			IStageDataProvider::SendMessage<FFramePerformanceProviderMessage>(EStageMessageFlags::None, GetFramePerformanceData());
@@ -134,8 +153,12 @@ public:
 
 	FFramePerformanceProviderMessage GetFramePerformanceData()
 	{
-		FFramePerformanceProviderMessage OutboundData =
-			UE::FramePerformanceProvider::Private::GetLatestPerformanceData(Stats);
+#if STATS
+		FFramePerformanceProviderMessage OutboundData = UE::FramePerformanceProvider::Private::GetLatestPerformanceData(Stats);
+#else
+		FFramePerformanceProviderMessage OutboundData = UE::FramePerformanceProvider::Private::GetLatestPerformanceData();
+#endif // STATS
+
 		if (OutboundData.CompilationTasksRemaining > 0)
 		{
 			OutboundData.Status = EStageMonitorNodeStatus::AssetCompiling;
@@ -194,7 +217,10 @@ public:
 	mutable FCriticalSection LoadInfoCS;
 
 	TArray<FLoadInfo> LoadStack;
+
+#if STATS
 	TArray<FStatMessage> Stats;
+#endif // STATS
 
 	const float UpdateRHIResourcesFrequency = 4.0f;
 	float UpdateFrequency = 0.2f; // Default is 200ms;
diff --git a/Engine/Plugins/VirtualProduction/StageMonitoring/StageMonitoring.uplugin b/Engine/Plugins/VirtualProduction/StageMonitoring/StageMonitoring.uplugin
--- a/Engine/Plugins/VirtualProduction/StageMonitoring/StageMonitoring.uplugin
+++ b/Engine/Plugins/VirtualProduction/StageMonitoring/StageMonitoring.uplugin
@@ -36,12 +36,12 @@
 		},
 		{
 			"Name": "StageDataProvider",
-			"Type": "UncookedOnly",
+			"Type": "Runtime",
 			"LoadingPhase": "Default"
 		},
 		{
 			"Name": "StageMonitorCommon",
-			"Type": "UncookedOnly",
+			"Type": "Runtime",
 			"LoadingPhase": "Default"
 		}
 	]
