StageMonitor: Make Stage Name column visible by default.

It is the common name used to identify nodes, so should be visible by default.

#rnx

diff --git a/Engine/Plugins/VirtualProduction/StageMonitoring/Source/StageMonitorEditor/Private/Widgets/SDataProviderListView.cpp b/Engine/Plugins/VirtualProduction/StageMonitoring/Source/StageMonitorEditor/Private/Widgets/SDataProviderListView.cpp
--- a/Engine/Plugins/VirtualProduction/StageMonitoring/Source/StageMonitorEditor/Private/Widgets/SDataProviderListView.cpp
+++ b/Engine/Plugins/VirtualProduction/StageMonitoring/Source/StageMonitorEditor/Private/Widgets/SDataProviderListView.cpp
@@ -233,7 +233,7 @@ namespace DataProviderListView
 			 }},
 			{StageName,
 			 {
-				 LOCTEXT("HeaderName_StageName", "Stage Name"), IsSortable, IsNotVisible,
+				 LOCTEXT("HeaderName_StageName", "Stage Name"), IsSortable, IsVisible,
 				 [](SDataProviderTableRow* Provider) -> TSharedRef<SWidget>
 				 {
 					 return SNew(SHorizontalBox)
