Temporary fix for ensure in FramePerformanceProvider.cpp
- FApp::GetIdleTime only works when called from the game thread, or task decendents of the game thread.
- This particular function is run on a separate worker thread, so is unsynchronized with the game thread time context. This system needs a refactor.

#Synced-CL 48346171

diff --git a/Engine/Plugins/VirtualProduction/StageMonitoring/Source/StageDataProvider/Private/FramePerformanceProvider.cpp b/Engine/Plugins/VirtualProduction/StageMonitoring/Source/StageDataProvider/Private/FramePerformanceProvider.cpp
--- a/Engine/Plugins/VirtualProduction/StageMonitoring/Source/StageDataProvider/Private/FramePerformanceProvider.cpp
+++ b/Engine/Plugins/VirtualProduction/StageMonitoring/Source/StageDataProvider/Private/FramePerformanceProvider.cpp
@@ -61,7 +61,10 @@ FFramePerformanceProviderMessage GetLatestPerformanceData()
 	const float GameThreadWaitTime = FPlatformTime::ToMilliseconds(GGameThreadWaitTime);
 	const float RenderThreadTime = FPlatformTime::ToMilliseconds(GRenderThreadTime);
 	const float RenderThreadWaitTime = FPlatformTime::ToMilliseconds(GRenderThreadWaitTime);
-	const float IdleTimeMilli = (FApp::GetIdleTime() * 1000.0);
+
+	// @todo delta time refactor - The FApp::GetIdleTime function requires proper pipelining of the delta time values from the game thread. Disabling this value for now.
+	// The above stats are also not threadsafe.
+	const float IdleTimeMilli = 0; // (FApp::GetIdleTime() * 1000.0);
 
 #if STATS
 	FPlatformMemoryStats Stats = FPlatformMemory::GetStats();
